Fixed extra synchronous loads in animation sharing

diff --git a/Engine/Plugins/Developer/AnimationSharing/Source/AnimationSharing/Private/AnimationSharingManager.cpp b/Engine/Plugins/Developer/AnimationSharing/Source/AnimationSharing/Private/AnimationSharingManager.cpp
--- a/Engine/Plugins/Developer/AnimationSharing/Source/AnimationSharing/Private/AnimationSharingManager.cpp
+++ b/Engine/Plugins/Developer/AnimationSharing/Source/AnimationSharing/Private/AnimationSharingManager.cpp
@@ -220,7 +220,7 @@ const FAnimationSharingScalability& UAnimationSharingManager::GetScalabilitySett
 
 void UAnimationSharingManager::SetupPerSkeletonData(const FPerSkeletonAnimationSharingSetup& SkeletonSetup)
 {
-	const USkeleton* Skeleton = SkeletonSetup.Skeleton.LoadSynchronous();
+	const USkeleton* Skeleton = SkeletonSetup.Skeleton;
 	UAnimationSharingStateProcessor* Processor = SkeletonSetup.StateProcessorClass ?SkeletonSetup.StateProcessorClass->GetDefaultObject<UAnimationSharingStateProcessor>() : nullptr;
 	UEnum* StateEnum = Processor ? Processor->GetAnimationStateEnum() : nullptr;
 	if (Skeleton && StateEnum && Processor)
@@ -749,7 +749,7 @@ uint8 UAnimSharingInstance::DetermineStateForActor(uint32 ActorIndex, bool& bSho
 
 bool UAnimSharingInstance::Setup(UAnimationSharingManager* AnimationSharingManager, const FPerSkeletonAnimationSharingSetup& SkeletonSetup, const FAnimationSharingScalability* InScalabilitySettings, uint32 Index)
 {
-	USkeletalMesh* SkeletalMesh = SkeletonSetup.SkeletalMesh.LoadSynchronous();
+	USkeletalMesh* SkeletalMesh = SkeletonSetup.SkeletalMesh;
 	/** Retrieve the state processor to use */
 	if (UAnimationSharingStateProcessor* Processor = SkeletonSetup.StateProcessorClass.GetDefaultObject())
 	{
@@ -846,7 +846,7 @@ void UAnimSharingInstance::SetupState(FPerStateData& StateData, const FAnimation
 	/** Setup overall data and flags */
 	StateData.bIsOnDemand = StateEntry.bOnDemand;
 	StateData.bIsAdditive = StateEntry.bAdditive;
-	StateData.AdditiveAnimationSequence = (StateEntry.bAdditive && StateEntry.AnimationSetups.IsValidIndex(0)) ? StateEntry.AnimationSetups[0].AnimSequence.LoadSynchronous() : nullptr;
+	StateData.AdditiveAnimationSequence = (StateEntry.bAdditive && StateEntry.AnimationSetups.IsValidIndex(0)) ? StateEntry.AnimationSetups[0].AnimSequence : nullptr;
 
 	/** Keep hard reference to animation sequence */
 	if (StateData.AdditiveAnimationSequence)
@@ -890,7 +890,7 @@ void UAnimSharingInstance::SetupState(FPerStateData& StateData, const FAnimation
 		const FAnimationSetup& AnimationSetup = StateEntry.AnimationSetups[SetupIndex];
 		/** User can setup either an AnimBP or AnimationSequence */
 		UClass* AnimBPClass = AnimationSetup.AnimBlueprint.Get();
-		UAnimSequence* AnimSequence = AnimationSetup.AnimSequence.LoadSynchronous();
+		UAnimSequence* AnimSequence = AnimationSetup.AnimSequence;
 		
 		if (AnimBPClass == nullptr && AnimSequence == nullptr)
 		{
diff --git a/Engine/Plugins/Developer/AnimationSharing/Source/AnimationSharing/Private/AnimationSharingSetup.cpp b/Engine/Plugins/Developer/AnimationSharing/Source/AnimationSharing/Private/AnimationSharingSetup.cpp
--- a/Engine/Plugins/Developer/AnimationSharing/Source/AnimationSharing/Private/AnimationSharingSetup.cpp
+++ b/Engine/Plugins/Developer/AnimationSharing/Source/AnimationSharing/Private/AnimationSharingSetup.cpp
@@ -6,25 +6,3 @@ UAnimationSharingSetup::UAnimationSharingSetup(const FObjectInitializer& ObjectI
 {
 	
 }
-
-#if WITH_EDITOR
-void UAnimationSharingSetup::PostLoad()
-{
-	Super::PostLoad();
-
-	/** Ensure all data required for the UI is loaded */
-	for (const FPerSkeletonAnimationSharingSetup& SharingSetup : SkeletonSetups)
-	{ 
-		SharingSetup.Skeleton.LoadSynchronous();
-		SharingSetup.SkeletalMesh.LoadSynchronous();
-		
-		for (const FAnimationStateEntry& Entry : SharingSetup.AnimationStates)
-		{
-			for (const FAnimationSetup& AnimSetup : Entry.AnimationSetups)
-			{
-				AnimSetup.AnimSequence.LoadSynchronous();				
-			}
-		}
-	}
-}
-#endif // WITH_EDITOR
diff --git a/Engine/Plugins/Developer/AnimationSharing/Source/AnimationSharing/Public/AnimationSharingManager.h b/Engine/Plugins/Developer/AnimationSharing/Source/AnimationSharing/Public/AnimationSharingManager.h
--- a/Engine/Plugins/Developer/AnimationSharing/Source/AnimationSharing/Public/AnimationSharingManager.h
+++ b/Engine/Plugins/Developer/AnimationSharing/Source/AnimationSharing/Public/AnimationSharingManager.h
@@ -531,6 +531,7 @@ protected:
 	FPerActorData* GetActorDataByHandle(uint32 InHandle);
 protected:
 	/** Array of unique skeletons, matches PerSkeletonData array entries*/
+	UPROPERTY(Transient)
 	TArray<const USkeleton*> Skeletons;
 
 	/** Sharing data required for the unique Skeleton setups */
diff --git a/Engine/Plugins/Developer/AnimationSharing/Source/AnimationSharing/Public/AnimationSharingSetup.h b/Engine/Plugins/Developer/AnimationSharing/Source/AnimationSharing/Public/AnimationSharingSetup.h
--- a/Engine/Plugins/Developer/AnimationSharing/Source/AnimationSharing/Public/AnimationSharingSetup.h
+++ b/Engine/Plugins/Developer/AnimationSharing/Source/AnimationSharing/Public/AnimationSharingSetup.h
@@ -15,10 +15,6 @@ class ANIMATIONSHARING_API UAnimationSharingSetup : public UObject
 	GENERATED_UCLASS_BODY()
 public:
 
-#if WITH_EDITOR
-	virtual void PostLoad() override;
-#endif // WITH_EDITOR
-
 	UPROPERTY(EditAnywhere, config, Category = AnimationSharing)
 	TArray<FPerSkeletonAnimationSharingSetup> SkeletonSetups;
 
diff --git a/Engine/Plugins/Developer/AnimationSharing/Source/AnimationSharing/Public/AnimationSharingTypes.h b/Engine/Plugins/Developer/AnimationSharing/Source/AnimationSharing/Public/AnimationSharingTypes.h
--- a/Engine/Plugins/Developer/AnimationSharing/Source/AnimationSharing/Public/AnimationSharingTypes.h
+++ b/Engine/Plugins/Developer/AnimationSharing/Source/AnimationSharing/Public/AnimationSharingTypes.h
@@ -22,7 +22,7 @@ public:
 	
 	/** Animation Sequence to play for this particular setup */
 	UPROPERTY(EditAnywhere, Category = AnimationSharing)
-	TSoftObjectPtr<UAnimSequence> AnimSequence;
+	UAnimSequence* AnimSequence;
 
 	/** Animation blueprint to use for playing back the Animation Sequence */
 	UPROPERTY(EditAnywhere, AdvancedDisplay, Category = AnimationSharing)
@@ -132,11 +132,11 @@ public:
 
 	/** Skeleton compatible with the animation sharing setup */
 	UPROPERTY(EditAnywhere, Category = AnimationSharing)
-	TSoftObjectPtr<USkeleton> Skeleton;
+	USkeleton* Skeleton;
 
 	/** Skeletal mesh used to setup skeletal mesh components */
 	UPROPERTY(EditAnywhere, Category = AnimationSharing)
-	TSoftObjectPtr<USkeletalMesh> SkeletalMesh;
+	USkeletalMesh* SkeletalMesh;
 
 	/** Animation blueprint used to perform the blending between states */
 	UPROPERTY(EditAnywhere, Category = AnimationSharing, meta = (DisplayName="Animation Blueprint for Blending"))
