Fix for crash with AnimInstanceSharing, the resource was not flagged as UPROPERTY and got deleted while we still had a native pointer to it, becoming garbage.

Should fix bugs with callstack ending with:
E4Editor-AnimationSharing.dll!UAnimSharingStateInstance::GetInstancedActors(TArray<AActor *,FDefaultAllocator> & Actors) Line 20   C++

#rnx

diff --git a/Engine/Plugins/Developer/AnimationSharing/Source/AnimationSharing/Public/AnimationSharingInstances.h b/Engine/Plugins/Developer/AnimationSharing/Source/AnimationSharing/Public/AnimationSharingInstances.h
--- a/Engine/Plugins/Developer/AnimationSharing/Source/AnimationSharing/Public/AnimationSharingInstances.h
+++ b/Engine/Plugins/Developer/AnimationSharing/Source/AnimationSharing/Public/AnimationSharingInstances.h
@@ -33,6 +33,8 @@ protected:
 private:
 	uint8 StateIndex;
 	uint8 ComponentIndex;
+	
+	UPROPERTY(Transient)
 	class UAnimSharingInstance* Instance;
 };
 
