Make sure that condition for failing to set up animation sharing instances takes into account whether or not it is an additive one

diff --git a/Engine/Plugins/Developer/AnimationSharing/Source/AnimationSharing/Private/AnimationSharingManager.cpp b/Engine/Plugins/Developer/AnimationSharing/Source/AnimationSharing/Private/AnimationSharingManager.cpp
--- a/Engine/Plugins/Developer/AnimationSharing/Source/AnimationSharing/Private/AnimationSharingManager.cpp
+++ b/Engine/Plugins/Developer/AnimationSharing/Source/AnimationSharing/Private/AnimationSharingManager.cpp
@@ -788,7 +788,7 @@ bool UAnimSharingInstance::Setup(UAnimationSharingManager* AnimationSharingManag
 				SetupState(StateData, StateEntry, SkeletalMesh, SkeletonSetup, Index);
 
 				// Make sure we have at least one component set up
-				if (StateData.Components.Num() == 0)
+				if ((!StateData.bIsAdditive && StateData.Components.Num() == 0) ||(StateData.bIsAdditive && AdditiveInstanceStack.AvailableInstances.Num() == 0))
 				{
 					UE_LOG(LogAnimationSharing, Error, TEXT("No Components available for State %s"), *StateEnum->GetDisplayNameTextByValue(StateValue).ToString());
 					bErrors = true;
