Error: EditCondition parsing failed: Field name "bBlendTransitions" was not found in class "AnimationSharingScalability" when viewing Use Blend Transitions in an Animation Sharing Setup
#fix Remove edit conditions as the bool has been converted to a PerPlatform property

diff --git a/Engine/Plugins/Developer/AnimationSharing/Source/AnimationSharing/Public/AnimationSharingTypes.h b/Engine/Plugins/Developer/AnimationSharing/Source/AnimationSharing/Public/AnimationSharingTypes.h
--- a/Engine/Plugins/Developer/AnimationSharing/Source/AnimationSharing/Public/AnimationSharingTypes.h
+++ b/Engine/Plugins/Developer/AnimationSharing/Source/AnimationSharing/Public/AnimationSharingTypes.h
@@ -167,15 +167,15 @@ public:
 	FPerPlatformBool UseBlendTransitions;
 
 	/** Significance value tied to whether or not a transition should be blended */
-	UPROPERTY(EditAnywhere, Category = AnimationSharing, meta = (EditCondition = "bBlendTransitions", ClampMin="0.0", UIMin="0.0"))
+	UPROPERTY(EditAnywhere, Category = AnimationSharing, meta = (ClampMin="0.0", UIMin="0.0"))
 	FPerPlatformFloat BlendSignificanceValue;
 
 	/** Maximum number of blends which can be running concurrently */
-	UPROPERTY(EditAnywhere, Category = AnimationSharing, meta = (EditCondition = "bBlendTransitions", ClampMin = "1", UIMin = "1"))
+	UPROPERTY(EditAnywhere, Category = AnimationSharing, meta = (ClampMin = "1", UIMin = "1"))
 	FPerPlatformInt MaximumNumberConcurrentBlends;
 
 	/** Significance value tied to whether or not the master pose components should be ticking */
-	UPROPERTY(EditAnywhere, Category = AnimationSharing, meta = (EditCondition = "bBlendTransitions", ClampMin = "0.0", UIMin = "0.0"))
+	UPROPERTY(EditAnywhere, Category = AnimationSharing, meta = (ClampMin = "0.0", UIMin = "0.0"))
 	FPerPlatformFloat TickSignificanceValue;
 };
 
