Animation sharing plugin:
- Moving Animation Sharing plugin file(s) to Engine
- Added debug materials to plugin and enabled plugin content
- Added BP exposed functionality for using the system
- Improved and or added detail customization with the lack of a dedicated editor
- Cleaned up (on-screen) debug information
- Added post-load to AnimationSharingSetup to ensure assets are loaded for editing
- Re-ordered, cleaned up and commented UProperties in various type classes

diff --git a/Engine/Plugins/Developer/AnimationSharing/AnimationSharing.uplugin b/Engine/Plugins/Developer/AnimationSharing/AnimationSharing.uplugin
new file mode 100644
--- /dev/null
+++ b/Engine/Plugins/Developer/AnimationSharing/AnimationSharing.uplugin
@@ -0,0 +1,35 @@
+{
+	"FileVersion": 1,
+	"Version": 1,
+	"VersionName": "1.0",
+	"FriendlyName": "Animation Sharing",
+	"Description": "Plugin to create Shared Animation systems using the Master-Child pose functionality",
+	"Category": "Programming",
+	"CreatedBy": "Epic Games, Inc.",
+	"CreatedByURL": "http://epicgames.com",
+	"DocsURL": "",
+	"MarketplaceURL": "",
+	"SupportURL": "",
+	"EnabledByDefault": true,
+	"CanContainContent": true,
+	"IsBetaVersion": false,
+	"Installed": false,
+	"Modules": [
+		{
+		  "Name": "AnimationSharing",
+		  "Type": "Runtime",
+		  "LoadingPhase": "PreDefault"
+		},
+		{
+			"Name": "AnimationSharingEd",
+			"Type": "Editor",
+			"LoadingPhase": "Default"
+		}
+	],
+	"Plugins": [ 
+		{ 
+			"Name": "SignificanceManager", 
+			"Enabled": true 
+		} 
+	]
+}
\ No newline at end of file
diff --git a/Engine/Plugins/Developer/AnimationSharing/Source/AnimationSharing/AnimationSharing.Build.cs b/Engine/Plugins/Developer/AnimationSharing/Source/AnimationSharing/AnimationSharing.Build.cs
new file mode 100644
--- /dev/null
+++ b/Engine/Plugins/Developer/AnimationSharing/Source/AnimationSharing/AnimationSharing.Build.cs
@@ -0,0 +1,31 @@
+// Copyright 1998-2018 Epic Games, Inc. All Rights Reserved.
+
+using UnrealBuildTool;
+
+public class AnimationSharing : ModuleRules
+{
+	public AnimationSharing(ReadOnlyTargetRules Target) : base(Target)
+	{
+        PCHUsage = ModuleRules.PCHUsageMode.UseExplicitOrSharedPCHs;
+		
+        PrivateDependencyModuleNames.AddRange(
+			new string[] {
+				"Core",
+				"CoreUObject",
+				"Engine",
+				"EngineSettings",
+				"SignificanceManager"
+            }
+		);
+
+        if (Target.Type == TargetType.Editor)
+        {
+            PrivateDependencyModuleNames.Add("TargetPlatform");
+        }
+
+        PrivateIncludePaths.AddRange(
+            new string[] {
+                "AnimationSharing/Private"
+         });       
+    }
+}
diff --git a/Engine/Plugins/Developer/AnimationSharing/Source/AnimationSharing/Private/AdditiveAnimationInstance.cpp b/Engine/Plugins/Developer/AnimationSharing/Source/AnimationSharing/Private/AdditiveAnimationInstance.cpp
new file mode 100644
--- /dev/null
+++ b/Engine/Plugins/Developer/AnimationSharing/Source/AnimationSharing/Private/AdditiveAnimationInstance.cpp
@@ -0,0 +1,76 @@
+// Copyright 1998-2018 Epic Games, Inc. All Rights Reserved.
+
+#include "AdditiveAnimationInstance.h"
+#include "AnimationSharingInstances.h"
+#include "AnimationSharingManager.h"
+
+FAdditiveAnimationInstance::FAdditiveAnimationInstance() : SkeletalMeshComponent(nullptr), AdditiveInstance(nullptr), AdditiveAnimationSequence(nullptr), BaseComponent(nullptr), bLoopingState(false)
+{
+}
+
+void FAdditiveAnimationInstance::Initialise(USkeletalMeshComponent* InSkeletalMeshComponent, UClass* InAnimationBPClass)
+{
+	if (InSkeletalMeshComponent && InAnimationBPClass)
+	{
+		SkeletalMeshComponent = InSkeletalMeshComponent;
+		SkeletalMeshComponent->SetAnimInstanceClass(InAnimationBPClass);
+		SkeletalMeshComponent->SetComponentTickEnabled(false);	
+		SkeletalMeshComponent->SetForcedLOD(1);
+		SkeletalMeshComponent->VisibilityBasedAnimTickOption = EVisibilityBasedAnimTickOption::AlwaysTickPoseAndRefreshBones;
+		AdditiveInstance = Cast<UAnimSharingAdditiveInstance>(SkeletalMeshComponent->GetAnimInstance());
+	}
+}
+
+void FAdditiveAnimationInstance::Setup(USkeletalMeshComponent* InBaseComponent, UAnimSequence* InAnimSequence)
+{
+	UAnimationSharingManager::SetDebugMaterial(SkeletalMeshComponent, 1);
+	SkeletalMeshComponent->SetComponentTickEnabled(true);
+	if (AdditiveInstance)
+	{
+		AdditiveInstance->BaseComponent = BaseComponent = InBaseComponent;
+		AdditiveInstance->AdditiveAnimation = AdditiveAnimationSequence = InAnimSequence;
+		AdditiveInstance->Alpha = 1.0f;
+		AdditiveInstance->bStateBool = bLoopingState = true;
+
+		SkeletalMeshComponent->AddTickPrerequisiteComponent(BaseComponent);
+	}
+}
+
+void FAdditiveAnimationInstance::UpdateBaseComponent(USkeletalMeshComponent* InBaseComponent)
+{
+	if (AdditiveInstance)
+	{
+		SkeletalMeshComponent->RemoveTickPrerequisiteComponent(BaseComponent);
+		AdditiveInstance->BaseComponent = BaseComponent = InBaseComponent;
+		SkeletalMeshComponent->AddTickPrerequisiteComponent(BaseComponent);
+	}
+}
+
+void FAdditiveAnimationInstance::Stop()
+{
+	if (AdditiveInstance)
+	{
+		UAnimationSharingManager::SetDebugMaterial(SkeletalMeshComponent, 0);
+		SkeletalMeshComponent->SetComponentTickEnabled(false);
+		SkeletalMeshComponent->RemoveTickPrerequisiteComponent(BaseComponent);
+	}
+}
+
+void FAdditiveAnimationInstance::Start()
+{
+	if (AdditiveInstance)
+	{
+		AdditiveInstance->bStateBool = bLoopingState = false;
+	}
+}
+
+USkeletalMeshComponent* FAdditiveAnimationInstance::GetBaseComponent() const
+{
+	return BaseComponent;
+}
+
+USkeletalMeshComponent* FAdditiveAnimationInstance::GetComponent() const
+{
+	return SkeletalMeshComponent;
+}
+
diff --git a/Engine/Plugins/Developer/AnimationSharing/Source/AnimationSharing/Private/AnimationSharingInstances.cpp b/Engine/Plugins/Developer/AnimationSharing/Source/AnimationSharing/Private/AnimationSharingInstances.cpp
new file mode 100644
--- /dev/null
+++ b/Engine/Plugins/Developer/AnimationSharing/Source/AnimationSharing/Private/AnimationSharingInstances.cpp
@@ -0,0 +1,47 @@
+// Copyright 1998-2017 Epic Games, Inc. All Rights Reserved.
+
+#include "AnimationSharingInstances.h"
+#include "AnimationSharingManager.h"
+#include "Algo/Transform.h"
+#include "Stats/Stats.h"
+
+UAnimSharingStateInstance::UAnimSharingStateInstance(const FObjectInitializer& ObjectInitializer)
+	: Super(ObjectInitializer), AnimationToPlay(nullptr), PermutationTimeOffset(0.f), PlayRate(1.f), StateIndex(INDEX_NONE), ComponentIndex(INDEX_NONE), Instance(nullptr)
+{
+	
+}
+
+void UAnimSharingStateInstance::GetInstancedActors(TArray<class AActor*>& Actors)
+{
+	QUICK_SCOPE_CYCLE_COUNTER(STAT_GetInstancedActors);
+	if (Instance && Instance->PerStateData.IsValidIndex(StateIndex))
+	{
+		FPerStateData& StateData = Instance->PerStateData[StateIndex];
+		if (StateData.Components.IsValidIndex(ComponentIndex))
+		{
+			USkeletalMeshComponent* Component = StateData.Components[ComponentIndex];
+			const TArray<TWeakObjectPtr<USkinnedMeshComponent>>& SlaveComponents = Component->GetSlavePoseComponents();
+
+			Algo::TransformIf(SlaveComponents, Actors,
+				[&Actors](const TWeakObjectPtr<USkinnedMeshComponent>& WeakPtr)
+			{
+				// Needs to be valid and unique
+				return WeakPtr.IsValid() && !Actors.Contains(WeakPtr->GetOwner());
+			},
+				[](const TWeakObjectPtr<USkinnedMeshComponent>& WeakPtr)
+			{
+				return WeakPtr->GetOwner();
+			});
+		}
+	}
+}
+
+UAnimSharingTransitionInstance::UAnimSharingTransitionInstance(const FObjectInitializer& ObjectInitializer)
+	: Super(ObjectInitializer),  FromComponent(nullptr), ToComponent(nullptr), BlendTime(.5f), bBlendBool(false)
+{
+}
+
+UAnimSharingAdditiveInstance::UAnimSharingAdditiveInstance(const FObjectInitializer& ObjectInitializer)
+	: Super(ObjectInitializer)
+{
+}
diff --git a/Engine/Plugins/Developer/AnimationSharing/Source/AnimationSharing/Private/AnimationSharingManager.cpp b/Engine/Plugins/Developer/AnimationSharing/Source/AnimationSharing/Private/AnimationSharingManager.cpp
new file mode 100644
--- /dev/null
+++ b/Engine/Plugins/Developer/AnimationSharing/Source/AnimationSharing/Private/AnimationSharingManager.cpp
@@ -0,0 +1,2102 @@
+// Copyright 1998-2018 Epic Games, Inc. All Rights Reserved.
+
+#include "AnimationSharingManager.h"
+#include "AnimationSharingModule.h"
+#include "Stats/Stats.h"
+#include "Components/SkinnedMeshComponent.h"
+#include "TransitionBlendInstance.h"
+#include "Animation/AnimationAsset.h"
+#include "Animation/AnimSequence.h"
+#include "SignificanceManager.h"
+#include "AnimationSharingSetup.h"
+#include "AdditiveAnimationInstance.h"
+#include "AnimationSharingInstances.h"
+
+#include "Misc/CoreMisc.h"
+#include "DrawDebugHelpers.h"
+#include "Math/NumericLimits.h"
+#include "Logging/LogMacros.h"
+#include "Engine/Engine.h"
+#include "ProfilingDebugging/CsvProfiler.h"
+#include "Misc/DefaultValueHelper.h"
+
+#if WITH_EDITOR
+#include "PlatformInfo.h"
+#include "Interfaces/ITargetPlatformManagerModule.h"
+#include "Interfaces/ITargetPlatform.h"
+#endif // WITH_EDITOR
+
+DEFINE_LOG_CATEGORY(LogAnimationSharing);
+
+DECLARE_CYCLE_STAT(TEXT("Tick"), STAT_AnimationSharing_Tick, STATGROUP_AnimationSharing);
+DECLARE_CYCLE_STAT(TEXT("UpdateBlends"), STAT_AnimationSharing_UpdateBlends, STATGROUP_AnimationSharing);
+DECLARE_CYCLE_STAT(TEXT("UpdateOnDemands"), STAT_AnimationSharing_UpdateOnDemands, STATGROUP_AnimationSharing);
+DECLARE_CYCLE_STAT(TEXT("UpdateAdditives"), STAT_AnimationSharing_UpdateAdditives, STATGROUP_AnimationSharing);
+DECLARE_CYCLE_STAT(TEXT("TickActorStates"), STAT_AnimationSharing_TickActorStates, STATGROUP_AnimationSharing);
+DECLARE_CYCLE_STAT(TEXT("KickoffInstances"), STAT_AnimationSharing_KickoffInstances, STATGROUP_AnimationSharing);
+
+DECLARE_DWORD_COUNTER_STAT(TEXT("NumBlends"), STAT_AnimationSharing_NumBlends, STATGROUP_AnimationSharing);
+DECLARE_DWORD_COUNTER_STAT(TEXT("NumOnDemands"), STAT_AnimationSharing_NumOnDemands, STATGROUP_AnimationSharing);
+DECLARE_DWORD_COUNTER_STAT(TEXT("NumActors"), STAT_AnimationSharing_NumActors, STATGROUP_AnimationSharing);
+DECLARE_DWORD_COUNTER_STAT(TEXT("NumComponent"), STAT_AnimationSharing_NumComponent, STATGROUP_AnimationSharing);
+
+static int32 GAnimationSharingDebugging = 0;
+static FAutoConsoleVariableRef CVarAnimSharing_DebugStates(
+	TEXT("a.Sharing.DebugStates"),
+	GAnimationSharingDebugging,
+	TEXT("Values: 0/1/2/3\n")
+	TEXT("Controls whether and which animation sharing debug features are enabled.\n")
+	TEXT("0: Turned off.\n")
+	TEXT("1: Turns on active master-components and blend with material coloring, and printing state information for each actor above their capsule.\n")
+	TEXT("2: Turns printing state information about currently active animation states, blend etc. Also enables line drawing from slave-components to currently assigned master components."),
+	ECVF_Cheat);
+
+static int32 GAnimationSharingEnabled = 1;
+static FAutoConsoleCommandWithWorldAndArgs CVarAnimSharing_Enabled(
+	TEXT("a.Sharing.Enabled"),
+	TEXT("Arguments: 0/1\n")
+	TEXT("Controls whether the animation sharing is enabled."),
+	FConsoleCommandWithWorldAndArgsDelegate::CreateLambda([](const TArray<FString>& Args, UWorld* World)
+	{
+		if (Args.Num() != 0)
+		{
+			const bool bShouldBeEnabled = Args[0].ToBool();
+			if (!bShouldBeEnabled && GAnimationSharingEnabled && World)
+			{
+				/** Need to unregister actors here*/
+				UAnimationSharingManager* Manager = FAnimSharingModule::Get(World);
+				if (Manager)
+				{
+					Manager->ClearActorData();
+				}
+			}
+
+			GAnimationSharingEnabled = bShouldBeEnabled;
+			UE_LOG(LogAnimationSharing, Log, TEXT("Animation Sharing System - %s"), GAnimationSharingEnabled ? TEXT("Enabled") : TEXT("Disabled"));
+		}
+	}),
+	ECVF_Cheat);
+
+#if !UE_BUILD_SHIPPING
+static int32 GMasterComponentsVisible = 0;
+static FAutoConsoleCommandWithWorldAndArgs CVarAnimSharing_ToggleVisibility(
+	TEXT("a.Sharing.ToggleVisibility"),
+	TEXT("Toggles the visibility of the Master Pose Components."),
+	FConsoleCommandWithWorldAndArgsDelegate::CreateLambda([](const TArray<FString>& Args, UWorld* World)
+	{
+		const bool bShouldBeVisible = !GMasterComponentsVisible;
+
+		/** Need to unregister actors here*/
+		UAnimationSharingManager* Manager = FAnimSharingModule::Get(World);
+		if (Manager)
+		{
+			Manager->SetMasterComponentsVisibility(bShouldBeVisible);
+		}
+
+		GMasterComponentsVisible = bShouldBeVisible;
+	}),
+	ECVF_Cheat);
+
+#else
+static const int32 GMasterComponentsVisible = 0;
+#endif
+
+#if WITH_EDITOR
+static TAutoConsoleVariable<FString> CVarAnimSharing_PreviewScalabilityPlatform(
+	TEXT("a.Sharing.ScalabilityPlatform"),
+	"",
+	TEXT("Controls which platform should be used when retrieving per platform scalability settings.\n")
+	TEXT("Empty: Current platform.\n")
+	TEXT("Name of Platform\n")
+	TEXT("Name of Platform Group\n"),
+	ECVF_Cheat);
+#endif
+
+#define LOG_STATES 0
+#define CSV_STATS 0
+#define DETAIL_STATS 0
+
+#if DEBUG_MATERIALS
+TArray<UMaterialInterface*> UAnimationSharingManager::DebugMaterials;
+#endif
+
+void UAnimationSharingManager::BeginDestroy()
+{
+	Super::BeginDestroy();
+	PerSkeletonData.Empty();
+	
+	// Unregister tick function
+	TickFunction.UnRegisterTickFunction();
+	TickFunction.Manager = nullptr;
+}
+
+UWorld* UAnimationSharingManager::GetWorld() const
+{
+	return Cast<UWorld>(GetOuter());
+}
+
+UAnimationSharingManager* UAnimationSharingManager::GetAnimationSharingManager(UObject* WorldContextObject)
+{
+	if (UWorld* World = GEngine->GetWorldFromContextObject(WorldContextObject, EGetWorldErrorMode::LogAndReturnNull))
+	{
+		return GetManagerForWorld(World);
+	}
+
+	return nullptr;
+}
+
+UAnimationSharingManager* UAnimationSharingManager::GetManagerForWorld(UWorld* InWorld)
+{
+	return FAnimSharingModule::Get(InWorld);
+}
+
+void FTickAnimationSharingFunction::ExecuteTick(float DeltaTime, enum ELevelTick TickType, ENamedThreads::Type CurrentThread, const FGraphEventRef& MyCompletionGraphEvent)
+{
+	if (ensure(Manager))
+	{
+		Manager->Tick(DeltaTime);
+	}
+}
+
+FString FTickAnimationSharingFunction::DiagnosticMessage()
+{
+	return TEXT("FTickAnimationSharingFunction");
+}
+
+FTickAnimationSharingFunction& UAnimationSharingManager::GetTickFunction()
+{
+	return TickFunction;
+}
+
+void UAnimationSharingManager::Initialise(const UAnimationSharingSetup* InSetup)
+{
+	if (InSetup)
+	{
+		TickFunction.Manager = this;
+		TickFunction.RegisterTickFunction(GetWorld()->PersistentLevel);
+
+		ScalabilitySettings = InSetup->ScalabilitySettings;
+
+#if WITH_EDITOR
+		// Update local copy defaults with current platform value
+		const FName PlatformName = UAnimationSharingManager::GetPlatformName();
+		ScalabilitySettings.UseBlendTransitions = ScalabilitySettings.UseBlendTransitions.GetValueForPlatformIdentifiers(PlatformName, PlatformName);
+		ScalabilitySettings.BlendSignificanceValue = ScalabilitySettings.BlendSignificanceValue.GetValueForPlatformIdentifiers(PlatformName, PlatformName);
+		ScalabilitySettings.MaximumNumberConcurrentBlends = ScalabilitySettings.MaximumNumberConcurrentBlends.GetValueForPlatformIdentifiers(PlatformName, PlatformName);
+		ScalabilitySettings.TickSignificanceValue = ScalabilitySettings.TickSignificanceValue.GetValueForPlatformIdentifiers(PlatformName, PlatformName);
+#endif
+
+		// Debug materials
+#if DEBUG_MATERIALS 
+		DebugMaterials.Empty();
+		{
+			UMaterialInterface* RedMaterial = LoadObject<UMaterialInterface>(nullptr, TEXT("/AnimationSharing/AnimSharingRed.AnimSharingRed"));
+			DebugMaterials.Add(RedMaterial);
+			UMaterialInterface* GreenMaterial = LoadObject<UMaterialInterface>(nullptr, TEXT("/AnimationSharing/AnimSharingGreen.AnimSharingGreen"));
+			DebugMaterials.Add(GreenMaterial);
+			UMaterialInterface* BlueMaterial = LoadObject<UMaterialInterface>(nullptr, TEXT("/AnimationSharing/AnimSharingBlue.AnimSharingBlue"));
+			DebugMaterials.Add(BlueMaterial);
+		}
+#endif 
+		UWorld* World = GetWorld();
+
+		for (const FPerSkeletonAnimationSharingSetup& SkeletonSetup : InSetup->SkeletonSetups)
+		{
+			SetupPerSkeletonData(SkeletonSetup);
+		}
+	}
+}
+
+const FAnimationSharingScalability& UAnimationSharingManager::GetScalabilitySettings() const
+{
+	return ScalabilitySettings;
+}
+
+void UAnimationSharingManager::SetupPerSkeletonData(const FPerSkeletonAnimationSharingSetup& SkeletonSetup)
+{
+	const USkeleton* Skeleton = SkeletonSetup.Skeleton.LoadSynchronous();
+	UEnum* StateEnum = SkeletonSetup.StateProcessorClass->GetDefaultObject<UAnimationSharingStateProcessor>()->GetAnimationStateEnum();
+	if (Skeleton && StateEnum)
+	{
+		UAnimSharingInstance* Data = NewObject<UAnimSharingInstance>(this);
+		PerSkeletonData.Add(Data);
+		Skeletons.Add(Skeleton);
+		Data->Setup(this, SkeletonSetup, &ScalabilitySettings, Skeletons.Num() - 1);
+	}
+}
+
+uint32 UAnimationSharingManager::CreateActorHandle(uint8 SkeletonIndex, uint32 ActorIndex) const
+{
+	ensureMsgf(ActorIndex <= 0xFFFFFF, TEXT("Invalid Actor Handle due to overflow"));
+	return (SkeletonIndex << 24) | ActorIndex;
+}
+
+uint8 UAnimationSharingManager::GetSkeletonIndexFromHandle(uint32 InHandle) const
+{
+	return (InHandle & 0xFF000000) >> 24;
+}
+
+uint32 UAnimationSharingManager::GetActorIndexFromHandle(uint32 InHandle) const
+{
+	return (InHandle & 0x00FFFFFF);
+}
+
+void UAnimationSharingManager::Tick(float DeltaTime)
+{
+	SCOPE_CYCLE_COUNTER(STAT_AnimationSharing_Tick);
+	
+	const float WorldTime = GetWorld()->GetTimeSeconds();
+
+	/** Keeping track of currently running instances / animations for debugging purposes */
+	int32 TotalNumBlends = 0;
+	int32 TotalNumOnDemands = 0;
+	int32 TotalNumComponents = 0;
+	int32 TotalNumActors = 0;
+
+	int32 TotalNumRunningStates = 0;
+	int32 TotalNumRunningComponents = 0;
+
+	/** Iterator over all Skeleton setups */
+	for (int32 Index = 0; Index < PerSkeletonData.Num(); ++Index)
+	{
+		UAnimSharingInstance* Instance = PerSkeletonData[Index];
+		Instance->WorldTime = WorldTime;
+
+		/** Tick both Blend and On-Demand instances first, as they could be finishing */
+		Instance->TickBlendInstances();
+		Instance->TickOnDemandInstances();
+		Instance->TickAdditiveInstances();
+
+		/** Tick actor states */
+		Instance->TickActorStates();
+
+		/** Setup and start any blending transitions created while ticking the actor states */
+		Instance->KickoffInstances();
+
+#if !UE_BUILD_SHIPPING
+		if (GAnimationSharingDebugging >= 1)
+		{
+			Instance->TickDebugInformation();
+		}
+#endif
+		/** Tick the animation states to determine which components should be turned on/off */
+		Instance->TickAnimationStates();
+
+#if DETAIL_STATS
+		/** Stat counters */
+		TotalNumOnDemands += Instance->OnDemandInstances.Num();
+		TotalNumBlends += Instance->BlendInstances.Num();
+		TotalNumActors += Instance->PerActorData.Num();
+		TotalNumComponents += Instance->PerComponentData.Num();
+
+		for (FPerStateData& StateData : Instance->PerStateData)
+		{
+			if (StateData.InUseComponentFrameBits.Contains(true))
+			{
+				++TotalNumRunningStates;
+			}
+
+			for (int32 ComponentIndex = 0; ComponentIndex < StateData.PreviousInUseComponentFrameBits.Num(); ++ComponentIndex)
+			{
+				if (StateData.PreviousInUseComponentFrameBits[ComponentIndex] == true)
+				{
+					++TotalNumRunningComponents;
+				}
+			}
+		}
+#endif // DETAIL_STATS
+	}
+
+#if DETAIL_STATS
+	SET_DWORD_STAT(STAT_AnimationSharing_NumOnDemands, TotalNumOnDemands);
+	SET_DWORD_STAT(STAT_AnimationSharing_NumBlends, TotalNumBlends);
+
+	SET_DWORD_STAT(STAT_AnimationSharing_NumActors, TotalNumActors);
+	SET_DWORD_STAT(STAT_AnimationSharing_NumComponent, TotalNumComponents);
+
+	SET_DWORD_STAT(STAT_AnimationSharing_NumBlends, TotalNumBlends);
+#endif // DETAIL_STATS
+
+#if CSV_STATS 
+	CSV_CUSTOM_STAT_GLOBAL(NumOnDemands, TotalNumOnDemands, ECsvCustomStatOp::Set);
+	CSV_CUSTOM_STAT_GLOBAL(NumBlends, TotalNumBlends, ECsvCustomStatOp::Set);
+	CSV_CUSTOM_STAT_GLOBAL(NumRunningStates, TotalNumRunningStates, ECsvCustomStatOp::Set);
+	CSV_CUSTOM_STAT_GLOBAL(NumRunningComponents, TotalNumRunningComponents, ECsvCustomStatOp::Set);
+#endif
+
+}
+
+void UAnimationSharingManager::RegisterActor(AActor* InActor, FUpdateActorHandle CallbackDelegate)
+{
+	if (!UAnimationSharingManager::AnimationSharingEnabled())
+	{
+		return;
+	}
+
+	if (InActor)
+	{
+		TArray<USkeletalMeshComponent*, TInlineAllocator<1>> OwnedComponents;
+		InActor->GetComponents<USkeletalMeshComponent>(OwnedComponents);
+		checkf(OwnedComponents.Num(), TEXT("No SkeletalMeshComponents found in actor!"));
+
+		const USkeleton* UsedSkeleton = [&]()
+		{
+			USkeleton* CurrentSkeleton = nullptr;
+			for (USkeletalMeshComponent* SkeletalMeshComponent : OwnedComponents)
+			{
+				const USkeletalMesh* Mesh = SkeletalMeshComponent->SkeletalMesh;
+				USkeleton* Skeleton = Mesh->Skeleton;
+
+				if (CurrentSkeleton == nullptr)
+				{
+					CurrentSkeleton = Skeleton;
+				}
+				else if (CurrentSkeleton != Skeleton)
+				{
+					if (!CurrentSkeleton->IsCompatibleMesh(Mesh))
+					{
+						checkf(false, TEXT("Multiple different skeletons within same actor"));
+					}
+				}
+			}
+
+			return CurrentSkeleton;
+		}();
+
+		RegisterActorWithSkeleton(InActor, UsedSkeleton, CallbackDelegate);
+	}
+}
+
+void UAnimationSharingManager::RegisterActorWithSkeleton(AActor* InActor, const USkeleton* SharingSkeleton, FUpdateActorHandle CallbackDelegate)
+{
+	if (!UAnimationSharingManager::AnimationSharingEnabled())
+	{
+		return;
+	}
+
+	const AnimationSharingDataHandle Handle = [&]() -> uint32
+	{
+		uint32 ArrayIndex = Skeletons.IndexOfByPredicate([&](const USkeleton* Skeleton)
+		{
+			return (Skeleton == SharingSkeleton) || (Skeleton->IsCompatible(SharingSkeleton));
+		});
+		ensureMsgf(ArrayIndex != INDEX_NONE, TEXT("Invalid skeleton for which there is no sharing setup available!"));
+		return ArrayIndex;
+	}();
+
+	if (Handle != INDEX_NONE)
+	{
+		TArray<USkeletalMeshComponent*, TInlineAllocator<1>> OwnedComponents;
+		InActor->GetComponents<USkeletalMeshComponent>(OwnedComponents);
+		checkf(OwnedComponents.Num(), TEXT("No SkeletalMeshComponents found in actor!"));
+
+		UAnimSharingInstance* Data = PerSkeletonData[Handle];
+
+		// Register the actor
+		const int32 ActorIndex = Data->RegisteredActors.Add(InActor);
+
+		FPerActorData& ActorData = Data->PerActorData.AddZeroed_GetRef();
+		ActorData.BlendInstanceIndex = ActorData.OnDemandInstanceIndex = ActorData.AdditiveInstanceIndex = INDEX_NONE;
+		ActorData.SignificanceValue = Data->SignificanceManager->GetSignificance(InActor);
+		ActorData.UpdateActorHandleDelegate = CallbackDelegate;
+
+		bool bShouldProcess = true;
+		ActorData.CurrentState = ActorData.PreviousState = Data->DetermineStateForActor(ActorIndex, bShouldProcess);
+
+		for (USkeletalMeshComponent* Component : OwnedComponents)
+		{
+			FPerComponentData& ComponentData = Data->PerComponentData.AddZeroed_GetRef();
+			ComponentData.ActorIndex = ActorIndex;
+			ComponentData.Component = Component;
+
+			Component->PrimaryComponentTick.bCanEverTick = false;
+			Component->SetComponentTickEnabled(false);
+			Component->bIgnoreMasterPoseComponentLOD = true;
+
+			ActorData.ComponentIndices.Add(Data->PerComponentData.Num() - 1);
+
+			const int32 ComponentIndex = Data->PerComponentData.Num() - 1;
+			Data->SetupSlaveComponent(ActorData.CurrentState, ActorIndex);
+		}
+
+		if (Data->PerStateData[ActorData.CurrentState].bIsOnDemand && ActorData.OnDemandInstanceIndex != INDEX_NONE)
+		{
+			// We will have setup an on-demand instance so we need to kick it off here before we next tick
+			Data->OnDemandInstances[ActorData.OnDemandInstanceIndex].bActive = true;
+			Data->OnDemandInstances[ActorData.OnDemandInstanceIndex].StartTime = Data->WorldTime;
+		}
+
+		const int32 ActorHandle = CreateActorHandle(Handle, ActorIndex);
+		ActorData.UpdateActorHandleDelegate.ExecuteIfBound(ActorHandle);
+	}
+}
+
+void UAnimationSharingManager::RegisterActorWithSkeletonBP(AActor* InActor, const USkeleton* SharingSkeleton)
+{
+	RegisterActorWithSkeleton(InActor, SharingSkeleton, FUpdateActorHandle::CreateLambda([](int32 A) {}));
+}
+
+void UAnimationSharingManager::UnregisterActor(AActor* InActor)
+{
+	if (!UAnimationSharingManager::AnimationSharingEnabled())
+	{
+		return;
+	}
+
+	for (int32 SkeletonIndex = 0; SkeletonIndex < PerSkeletonData.Num(); ++SkeletonIndex)
+	{
+		UAnimSharingInstance* SkeletonData = PerSkeletonData[SkeletonIndex];
+		const int32 ActorIndex = SkeletonData->RegisteredActors.IndexOfByKey(InActor);	
+
+		if (ActorIndex != INDEX_NONE )
+		{
+			const FPerActorData& ActorData = SkeletonData->PerActorData[ActorIndex];
+
+			const bool bNeedsSwap = SkeletonData->PerActorData.Num() > 1 && ActorIndex != SkeletonData->PerActorData.Num() - 1;
+
+			for (int32 ComponentIndex : ActorData.ComponentIndices)
+			{
+				SkeletonData->PerComponentData[ComponentIndex].Component->SetMasterPoseComponent(nullptr, true);
+				SkeletonData->PerComponentData[ComponentIndex].Component->SetComponentTickEnabled(true);
+				SkeletonData->RemoveComponent(ComponentIndex);
+			}
+
+			const int32 SwapIndex = SkeletonData->PerActorData.Num() - 1;
+
+			// Remove actor index from any blend instances
+			for (FBlendInstance& Instance : SkeletonData->BlendInstances)
+			{
+				Instance.ActorIndices.Remove(ActorIndex);
+
+				// If we are swapping and the actor we are swapping with is part of the instance make sure we update the actor index
+				const uint32 SwapActorIndex = bNeedsSwap ? Instance.ActorIndices.IndexOfByKey(SwapIndex) : INDEX_NONE;
+				if (SwapActorIndex != INDEX_NONE)
+				{
+					Instance.ActorIndices[SwapActorIndex] = ActorIndex;
+				}
+			}
+
+			// Remove actor index from any running on demand instances
+			for (FOnDemandInstance& Instance : SkeletonData->OnDemandInstances)
+			{
+				Instance.ActorIndices.Remove(ActorIndex);
+
+				// If we are swapping and the actor we are swapping with is part of the instance make sure we update the actor index
+				const uint32 SwapActorIndex = bNeedsSwap ? Instance.ActorIndices.IndexOfByKey(SwapIndex) : INDEX_NONE;
+				if (SwapActorIndex != INDEX_NONE)
+				{
+					Instance.ActorIndices[SwapActorIndex] = ActorIndex;
+				}
+			}
+
+			// Remove actor index from any additive instances
+			for (FAdditiveInstance& Instance : SkeletonData->AdditiveInstances)
+			{
+				if (Instance.ActorIndex == ActorIndex)
+				{
+					Instance.ActorIndex = INDEX_NONE;
+				}
+				else if (bNeedsSwap && Instance.ActorIndex == SwapIndex)
+				{
+					Instance.ActorIndex = ActorIndex;
+				}
+			}
+
+			if (bNeedsSwap)
+			{
+				// Swap actor index for all components which are part of the actor we are swapping with
+				for (uint32 ComponentIndex : SkeletonData->PerActorData[SwapIndex].ComponentIndices)
+				{
+					SkeletonData->PerComponentData[ComponentIndex].ActorIndex = ActorIndex;
+				}
+
+				// Make sure we update the handle on the swapped actor
+				SkeletonData->PerActorData[SwapIndex].UpdateActorHandleDelegate.ExecuteIfBound(CreateActorHandle(SkeletonIndex, ActorIndex));
+			}			
+
+			SkeletonData->PerActorData.RemoveAtSwap(ActorIndex, 1, false);
+			SkeletonData->RegisteredActors.RemoveAtSwap(ActorIndex, 1, false);
+		}
+	}
+}
+
+void UAnimationSharingManager::UpdateSignificanceForActorHandle(uint32 InHandle, float InValue)
+{
+	// Retrieve actor
+	if (FPerActorData* ActorData = GetActorDataByHandle(InHandle))
+	{
+		ActorData->SignificanceValue = InValue;
+	}
+}
+
+FPerActorData* UAnimationSharingManager::GetActorDataByHandle(uint32 InHandle)
+{
+	FPerActorData* ActorDataPtr = nullptr;
+	uint8 SkeletonIndex = GetSkeletonIndexFromHandle(InHandle);
+	uint32 ActorIndex = GetActorIndexFromHandle(InHandle);
+	if (PerSkeletonData.IsValidIndex(SkeletonIndex))
+	{
+		if (PerSkeletonData[SkeletonIndex]->PerActorData.IsValidIndex(ActorIndex))
+		{
+			ActorDataPtr = &PerSkeletonData[SkeletonIndex]->PerActorData[ActorIndex];
+		}
+	}
+
+	return ActorDataPtr;
+}
+
+void UAnimationSharingManager::ClearActorData()
+{
+	UnregisterAllActors();
+
+	for (UAnimSharingInstance* Data : PerSkeletonData)
+	{		
+		Data->BlendInstances.Empty();
+		Data->OnDemandInstances.Empty();
+	}
+}
+
+void UAnimationSharingManager::UnregisterAllActors()
+{
+	for (UAnimSharingInstance* Data : PerSkeletonData)
+	{
+		for (int32 ActorIndex = 0; ActorIndex < Data->RegisteredActors.Num(); ++ActorIndex)
+		{
+			AActor* RegisteredActor = Data->RegisteredActors[ActorIndex];
+			if (RegisteredActor)
+			{
+				FPerActorData& ActorData = Data->PerActorData[ActorIndex];
+				for (int32 ComponentIndex : ActorData.ComponentIndices)
+				{
+					Data->PerComponentData[ComponentIndex].Component->SetMasterPoseComponent(nullptr, true);
+					Data->PerComponentData[ComponentIndex].Component->PrimaryComponentTick.bCanEverTick = true;
+					Data->PerComponentData[ComponentIndex].Component->SetComponentTickEnabled(true);
+					Data->PerComponentData[ComponentIndex].Component->bRecentlyRendered = false;
+				}
+				ActorData.ComponentIndices.Empty();
+			}
+		}
+
+		Data->PerActorData.Empty();
+		Data->PerComponentData.Empty();
+		Data->RegisteredActors.Empty();
+	}	
+}
+
+void UAnimationSharingManager::SetMasterComponentsVisibility(bool bVisible)
+{
+	for (UAnimSharingInstance* Data : PerSkeletonData)
+	{
+		for (FPerStateData& StateData : Data->PerStateData)
+		{
+			for (USkeletalMeshComponent* Component : StateData.Components)
+			{
+				Component->SetVisibility(bVisible);
+			}
+		}
+
+		for (FTransitionBlendInstance* Instance : Data->BlendInstanceStack.AvailableInstances)
+		{
+			Instance->GetComponent()->SetVisibility(bVisible);
+		}
+
+		for (FTransitionBlendInstance* Instance : Data->BlendInstanceStack.InUseInstances)
+		{
+			Instance->GetComponent()->SetVisibility(bVisible);
+		}
+
+		for (FAdditiveAnimationInstance* Instance : Data->AdditiveInstanceStack.AvailableInstances)
+		{
+			Instance->GetComponent()->SetVisibility(bVisible);
+		}
+
+		for (FAdditiveAnimationInstance* Instance : Data->AdditiveInstanceStack.InUseInstances)
+		{
+			Instance->GetComponent()->SetVisibility(bVisible);
+		}
+	}
+}
+
+bool UAnimationSharingManager::AnimationSharingEnabled()
+{
+	return GAnimationSharingEnabled == 1;
+}
+
+bool UAnimationSharingManager::CreateAnimationSharingManager(UObject* WorldContextObject, const UAnimationSharingSetup* Setup)
+{
+	if (UWorld* World = GEngine->GetWorldFromContextObject(WorldContextObject, EGetWorldErrorMode::LogAndReturnNull))
+	{
+		return FAnimSharingModule::CreateAnimationSharingManager(World, Setup);
+	}
+
+	return false;
+}
+
+void UAnimationSharingManager::SetDebugMaterial(USkeletalMeshComponent* Component, uint8 State)
+{
+#if DEBUG_MATERIALS 
+	if (GAnimationSharingDebugging >= 1 && DebugMaterials.IsValidIndex(State))
+	{
+		const int32 NumMaterials = Component->GetNumMaterials();
+		for (int32 Index = 0; Index < NumMaterials; ++Index)
+		{
+			Component->SetMaterial(Index, DebugMaterials[State]);
+		}
+	}
+#endif
+}
+
+void UAnimationSharingManager::SetDebugMaterialForActor(UAnimSharingInstance* Data, uint32 ActorIndex, uint8 State)
+{
+#if DEBUG_MATERIALS 
+	for (uint32 ComponentIndex : Data->PerActorData[ActorIndex].ComponentIndices)
+	{
+		SetDebugMaterial(Data->PerComponentData[ComponentIndex].Component, State);
+	}
+#endif
+}
+
+#if WITH_EDITOR
+FName UAnimationSharingManager::GetPlatformName()
+{
+	const FString PlatformString = CVarAnimSharing_PreviewScalabilityPlatform.GetValueOnAnyThread();
+	if (PlatformString.IsEmpty())
+	{
+		ITargetPlatform* CurrentPlatform = GetTargetPlatformManagerRef().GetRunningTargetPlatform();
+		return CurrentPlatform->GetPlatformInfo().PlatformGroupName;
+	}
+
+	FName PlatformNameFromString(*PlatformString);
+	return PlatformNameFromString;
+}
+#endif
+
+void UAnimSharingInstance::BeginDestroy()
+{
+	Super::BeginDestroy();
+	for (FPerActorData& ActorData : PerActorData)
+	{
+		for (uint32 ComponentIndex : ActorData.ComponentIndices)
+		{
+			PerComponentData[ComponentIndex].Component->SetMasterPoseComponent(nullptr, true);
+		}
+	}
+
+	RegisteredActors.Empty();
+	PerActorData.Empty();
+	PerComponentData.Empty();
+	PerStateData.Empty();
+	StateProcessor = nullptr;
+	StateEnum = nullptr;
+	BlendInstances.Empty();
+	OnDemandInstances.Empty();
+}
+
+uint8 UAnimSharingInstance::DetermineStateForActor(uint32 ActorIndex, bool& bShouldProcess)
+{
+	const FPerActorData& ActorData = PerActorData[ActorIndex];
+	int32 State = 0;
+	if (bNativeStateProcessor)
+	{
+		StateProcessor->ProcessActorState_Implementation(State, RegisteredActors[ActorIndex], ActorData.CurrentState, ActorData.OnDemandInstanceIndex != INDEX_NONE ? OnDemandInstances[ActorData.OnDemandInstanceIndex].State : INDEX_NONE, bShouldProcess);
+	}
+	else
+	{
+		StateProcessor->ProcessActorState(State, RegisteredActors[ActorIndex], ActorData.CurrentState, ActorData.OnDemandInstanceIndex != INDEX_NONE ? OnDemandInstances[ActorData.OnDemandInstanceIndex].State : INDEX_NONE, bShouldProcess);
+	}
+	
+	return FMath::Max(0, State);
+}
+
+void UAnimSharingInstance::Setup(UAnimationSharingManager* AnimationSharingManager, const FPerSkeletonAnimationSharingSetup& SkeletonSetup, const FAnimationSharingScalability* InScalabilitySettings, uint32 Index)
+{
+	USkeletalMesh* SkeletalMesh = SkeletonSetup.SkeletalMesh.LoadSynchronous();
+	/** Retrieve the state processor to use */
+	if (UAnimationSharingStateProcessor* Processor = SkeletonSetup.StateProcessorClass.GetDefaultObject())
+	{
+		StateProcessor = Processor;
+		bNativeStateProcessor = SkeletonSetup.StateProcessorClass->HasAnyClassFlags(CLASS_Native);
+	}
+
+	if (SkeletalMesh && StateProcessor)
+	{
+		SkeletalMeshBounds = SkeletalMesh->GetBounds().BoxExtent * 2;
+		ScalabilitySettings = InScalabilitySettings;
+		StateEnum = StateProcessor->GetAnimationStateEnum();
+		const uint32 NumStates = StateEnum->NumEnums();
+		PerStateData.AddDefaulted(NumStates);
+
+		UWorld* World = GetWorld();
+		SharingActor = World->SpawnActor<AActor>();
+		// Make sure the actor stays around when scrubbing through replays, states will be updated correctly in next tick 
+		SharingActor->bReplayRewindable = true;
+		SignificanceManager = USignificanceManager::Get<USignificanceManager>(World);
+		AnimSharingManager = AnimationSharingManager;
+
+		/** Create runtime data structures for unique animation states */
+		NumSetups = 0;
+		for (const FAnimationStateEntry& StateEntry : SkeletonSetup.AnimationStates)
+		{
+			const uint8 StateValue = StateEntry.State;
+			const uint32 StateIndex = StateEnum->GetIndexByValue(StateValue);
+
+			checkf(!PerStateData.FindByPredicate([&](const FPerStateData& State) { return State.StateEnumValue == StateValue; }), TEXT("State already defined"));
+
+			FPerStateData& StateData = PerStateData[StateIndex];
+			StateData.StateEnumValue = StateValue;
+			SetupState(StateData, StateEntry, SkeletalMesh, SkeletonSetup, Index);
+		}
+
+		/** Setup blend actors, if enabled*/
+		if (ScalabilitySettings->UseBlendTransitions.Default)
+		{
+			const uint32 TotalNumberOfBlendActorsRequired = ScalabilitySettings->MaximumNumberConcurrentBlends.Default;
+			const float ZOffset = Index * SkeletalMeshBounds.Z * 2.f;
+			for (uint32 BlendIndex = 0; BlendIndex < TotalNumberOfBlendActorsRequired; ++BlendIndex)
+			{
+				const FVector SpawnLocation(BlendIndex * SkeletalMeshBounds.X, 0.f, ZOffset + SkeletalMeshBounds.Z);
+				const FName BlendComponentName(*(SkeletalMesh->GetName() + TEXT("_BlendComponent") + FString::FromInt(BlendIndex)));
+				USkeletalMeshComponent* BlendComponent = NewObject<USkeletalMeshComponent>(SharingActor, BlendComponentName);
+				BlendComponent->RegisterComponent();
+				BlendComponent->SetRelativeLocation(SpawnLocation);
+				BlendComponent->SetSkeletalMesh(SkeletalMesh);
+				BlendComponent->SetVisibility(GMasterComponentsVisible == 1);
+
+				BlendComponent->PrimaryComponentTick.AddPrerequisite(AnimSharingManager, AnimSharingManager->GetTickFunction());
+
+				FTransitionBlendInstance* BlendActor = new FTransitionBlendInstance();
+				BlendActor->Initialise(BlendComponent, SkeletonSetup.BlendAnimBlueprint.Get());
+				BlendInstanceStack.AddInstance(BlendActor);
+			}
+		}
+	}
+	else
+	{
+		UE_LOG(LogAnimationSharing, Error, TEXT("Invalid Skeletal Mesh or State Processing Class"));
+	}
+}
+
+void UAnimSharingInstance::SetupState(FPerStateData& StateData, const FAnimationStateEntry& StateEntry, USkeletalMesh* SkeletalMesh, const FPerSkeletonAnimationSharingSetup& SkeletonSetup, uint32 Index)
+{
+	/** Used for placing components into rows / columns at origin for debugging purposes */
+	const float ZOffset = Index * SkeletalMeshBounds.Z * 2.f;
+
+	/** Setup overall data and flags */
+	StateData.bIsOnDemand = StateEntry.bOnDemand;
+	StateData.bIsAdditive = StateEntry.bAdditive;
+	StateData.AdditiveAnimationSequence = (StateEntry.bAdditive && StateEntry.AnimationSetups.IsValidIndex(0)) ? StateEntry.AnimationSetups[0].AnimSequence.LoadSynchronous() : nullptr;
+
+	/** Keep hard reference to animation sequence */
+	if (StateData.AdditiveAnimationSequence)
+	{
+		UsedAnimationSequences.Add(StateData.AdditiveAnimationSequence);
+	}
+
+	StateData.BlendTime = StateEntry.BlendTime;	
+	StateData.bReturnToPreviousState = StateEntry.bReturnToPreviousState;
+	StateData.bShouldForwardToState = StateEntry.bSetNextState;
+	StateData.ForwardStateValue = StateEntry.NextState;
+
+	int32 MaximumNumberOfConcurrentInstances = StateEntry.MaximumNumberOfConcurrentInstances.Default;
+#if WITH_EDITOR
+	const FName PlatformName = UAnimationSharingManager::GetPlatformName();
+	MaximumNumberOfConcurrentInstances = StateEntry.MaximumNumberOfConcurrentInstances.GetValueForPlatformIdentifiers(PlatformName, PlatformName);
+#endif
+
+	/** Ensure that we spread our number over the number of enabled setups */
+	const int32 NumInstancesPerSetup = [MaximumNumberOfConcurrentInstances, &StateEntry]()
+	{
+		int32 TotalEnabled = 0;
+		for (const FAnimationSetup& AnimationSetup : StateEntry.AnimationSetups)
+		{
+			bool bEnabled = AnimationSetup.Enabled.Default;
+#if WITH_EDITOR
+			const FName PlatformName = UAnimationSharingManager::GetPlatformName();
+			bEnabled = AnimationSetup.Enabled.GetValueForPlatformIdentifiers(PlatformName, PlatformName);
+#endif
+			TotalEnabled += bEnabled ? 1 : 0;
+		}
+
+		return (TotalEnabled > 0) ? FMath::CeilToInt((float)MaximumNumberOfConcurrentInstances / (float)TotalEnabled) : 0;
+	}();
+
+	UWorld* World = GetWorld();
+	/** Setup animations used for this state and the number of permutations */
+	TArray<USkeletalMeshComponent*>& Components = StateData.Components;
+	for (int32 SetupIndex = 0; SetupIndex < StateEntry.AnimationSetups.Num(); ++SetupIndex)
+	{
+		const FAnimationSetup& AnimationSetup = StateEntry.AnimationSetups[SetupIndex];
+		/** User can setup either an AnimBP or AnimationSequence */
+		UClass* AnimBPClass = AnimationSetup.AnimBlueprint.Get();
+		UAnimSequence* AnimSequence = AnimationSetup.AnimSequence.LoadSynchronous();
+		ensureMsgf(AnimBPClass != nullptr || AnimSequence != nullptr, TEXT("Animation setup without either an Animation Blueprint Class of Animation Sequence"));
+
+		bool bEnabled = AnimationSetup.Enabled.Default;
+#if WITH_EDITOR			
+		bEnabled = AnimationSetup.Enabled.GetValueForPlatformIdentifiers(PlatformName, PlatformName);
+#endif
+
+		/** Only create component if the setup is enabled for this platform and we have a valid animation asset */
+		if (bEnabled && (AnimBPClass || AnimSequence))
+		{
+			int32 NumRandomizedInstances = AnimationSetup.NumRandomizedInstances.Default;
+#if WITH_EDITOR			
+			NumRandomizedInstances = AnimationSetup.NumRandomizedInstances.GetValueForPlatformIdentifiers(PlatformName, PlatformName);
+#endif
+			const uint32 NumInstances = StateEntry.bOnDemand ? NumInstancesPerSetup	: NumRandomizedInstances;
+			for (uint32 InstanceIndex = 0; InstanceIndex < NumInstances; ++InstanceIndex)
+			{
+				if (!StateData.bIsAdditive)
+				{
+					const FName StateComponentName(*(SkeletalMesh->GetName() + TEXT("_") + StateEnum->GetNameStringByIndex(StateEntry.State) + FString::FromInt(SetupIndex) + FString::FromInt(InstanceIndex)));
+					USkeletalMeshComponent* Component = NewObject<USkeletalMeshComponent>(SharingActor, StateComponentName);
+					Component->RegisterComponent();
+					/** Arrange component in correct row / column */
+					Component->SetRelativeLocation(FVector(NumSetups * SkeletalMeshBounds.X, 0.f, ZOffset));
+					/** Set shared skeletal mesh */
+					Component->SetSkeletalMesh(SkeletalMesh);
+					Component->VisibilityBasedAnimTickOption = EVisibilityBasedAnimTickOption::AlwaysTickPoseAndRefreshBones;
+					Component->SetForcedLOD(1);
+					Component->SetVisibility(GMasterComponentsVisible == 1);
+					Component->bPropagateCurvesToSlaves = StateEntry.bRequiresCurves;
+
+					if (AnimBPClass != nullptr && AnimSequence != nullptr)
+					{
+						Component->SetAnimInstanceClass(AnimBPClass);
+						if (UAnimSharingStateInstance* AnimInstance = Cast<UAnimSharingStateInstance>(Component->GetAnimInstance()))
+						{
+							AnimInstance->AnimationToPlay = AnimSequence;
+							if (InstanceIndex > 0)
+							{
+								const float Steps = (AnimSequence->SequenceLength * 0.9f) / (NumInstances);
+								const float StartTimeOffset = Steps * InstanceIndex;
+								AnimInstance->PermutationTimeOffset = StartTimeOffset;
+							}
+
+							AnimInstance->PlayRate = StateData.bIsOnDemand ? 0.f : 1.0f;
+
+							AnimInstance->Instance = this;
+							AnimInstance->StateIndex = StateEntry.State;
+							AnimInstance->ComponentIndex = Components.Num();
+
+							/** Set the current animation length length */
+							StateData.AnimationLengths.Add(AnimSequence->SequenceLength);
+						}
+					}
+					else if (AnimSequence != nullptr)
+					{
+						Component->PlayAnimation(AnimSequence, true);
+
+						/** If this is an on-demand state we pause the animation as we'll want to start it from the beginning anytime we start an on-demand instance */
+						if (StateData.bIsOnDemand)
+						{
+							Component->Stop();
+						}
+						else
+						{
+							if (InstanceIndex > 0)
+							{
+								const float Steps = (AnimSequence->SequenceLength * 0.9f) / (NumInstances);
+								const float StartTimeOffset = Steps * InstanceIndex;
+								Component->SetPosition(StartTimeOffset, false);
+							}
+						}
+
+						/** Set the current animation length length */
+						StateData.AnimationLengths.Add(AnimSequence->SequenceLength);
+					}
+
+					/** Set material to red to indicate that it's not in use*/
+					UAnimationSharingManager::SetDebugMaterial(Component, 0);
+
+					Component->PrimaryComponentTick.AddPrerequisite(AnimSharingManager, AnimSharingManager->GetTickFunction());
+					Components.Add(Component);
+				}
+				else
+				{					
+					const FVector SpawnLocation(FVector(NumSetups * SkeletalMeshBounds.X, 0.f, ZOffset));
+					const FName AdditiveComponentName(*(SkeletalMesh->GetName() + TEXT("_") + StateEnum->GetNameStringByIndex(StateEntry.State) + FString::FromInt(InstanceIndex)));
+					USkeletalMeshComponent* AdditiveComponent = NewObject<USkeletalMeshComponent>(SharingActor, AdditiveComponentName);
+					AdditiveComponent->RegisterComponent();
+					AdditiveComponent->SetRelativeLocation(SpawnLocation);
+					AdditiveComponent->SetSkeletalMesh(SkeletalMesh);
+					AdditiveComponent->SetVisibility(GMasterComponentsVisible == 1);
+
+					AdditiveComponent->PrimaryComponentTick.AddPrerequisite(AnimSharingManager, AnimSharingManager->GetTickFunction());
+
+					FAdditiveAnimationInstance* AdditiveInstance = new FAdditiveAnimationInstance();
+					AdditiveInstance->Initialise(AdditiveComponent, SkeletonSetup.AdditiveAnimBlueprint.Get());
+					AdditiveInstanceStack.AddInstance(AdditiveInstance);
+				}
+				
+				++NumSetups;
+			}
+		}
+	}
+
+	float TotalLength = 0.f;
+	for (float Length : StateData.AnimationLengths)
+	{
+		TotalLength += Length;
+	}
+	const float AverageLength = (StateData.AnimationLengths.Num() > 0) ? TotalLength / FMath::Min((float)StateData.AnimationLengths.Num(), 1.f) : 0.f;
+	StateData.WiggleTime = AverageLength * StateEntry.WiggleTimePercentage;
+
+	/** Randomizes the order of Components so we actually hit different animations when running on demand */
+	if (StateData.bIsOnDemand && !StateData.bIsAdditive && StateEntry.AnimationSetups.Num() > 1)
+	{	
+		TArray<USkeletalMeshComponent*> RandomizedComponents;
+		while (Components.Num() > 0)
+		{
+			const int32 RandomIndex = FMath::RandRange(0, Components.Num() - 1);
+			RandomizedComponents.Add(Components[RandomIndex]);
+			Components.RemoveAt(RandomIndex, 1);
+		}
+
+		Components = RandomizedComponents;
+	}
+
+	/** Initialize component (previous frame) usage flags */
+	StateData.InUseComponentFrameBits.Init(false, Components.Num());
+	/** This should enforce turning off the components tick during the first frame */
+	StateData.PreviousInUseComponentFrameBits.Init(true, Components.Num());
+
+	StateData.SlaveTickRequiredFrameBits.Init(false, Components.Num());
+}
+
+void UAnimSharingInstance::TickDebugInformation()
+{
+#if !UE_BUILD_SHIPPING
+#if UE_BUILD_DEVELOPMENT
+	if (GMasterComponentsVisible && GAnimationSharingDebugging >= 2)
+	{
+		for (const FPerStateData& StateData : PerStateData)
+		{
+			for (int32 Index = 0; Index < StateData.InUseComponentFrameBits.Num(); ++Index)
+			{
+				const FString ComponentString = FString::Printf(TEXT("In Use %s - Required %s"), StateData.InUseComponentFrameBits[Index] ? TEXT("True") : TEXT("False"), StateData.SlaveTickRequiredFrameBits[Index] ? TEXT("True") : TEXT("False"));
+				DrawDebugString(GetWorld(), StateData.Components[Index]->GetComponentLocation() + FVector(0,0,StateData.Components[Index]->Bounds.BoxExtent.Z), ComponentString, nullptr, FColor::White, 0.016f, false);
+			}
+		}
+	}
+#endif // UE_BUILD_DEVELOPMENT
+	
+
+	for (int32 ActorIndex = 0; ActorIndex < RegisteredActors.Num(); ++ActorIndex)
+	{
+		// Non-const for DrawDebugString
+		AActor* Actor = RegisteredActors[ActorIndex];
+		if (Actor)
+		{
+			const FPerActorData& ActorData = PerActorData[ActorIndex];
+			const uint8 State = ActorData.CurrentState;
+
+			const FString StateString = [&]() -> FString
+			{
+				/** Check whether or not we are currently blending */
+				const uint32 BlendInstanceIndex = ActorData.BlendInstanceIndex;
+				if (BlendInstanceIndex != INDEX_NONE && BlendInstances.IsValidIndex(BlendInstanceIndex))
+				{
+					const float TimeLeft = BlendInstances[BlendInstanceIndex].BlendTime - (GetWorld()->GetTimeSeconds() - BlendInstances[BlendInstanceIndex].EndTime);
+					return FString::Printf(TEXT("Blending states - %s to %s [%1.3f] (%i)"), *StateEnum->GetDisplayNameTextByValue(BlendInstances[BlendInstanceIndex].StateFrom).ToString(), *StateEnum->GetDisplayNameTextByValue(BlendInstances[BlendInstanceIndex].StateTo).ToString(), TimeLeft, ActorData.BlendInstanceIndex);
+				}
+
+				/** Check if we are part of an on-demand instance */ 
+				const uint32 DemandInstanceIndex = ActorData.OnDemandInstanceIndex;
+				if (DemandInstanceIndex != INDEX_NONE && OnDemandInstances.IsValidIndex(DemandInstanceIndex))
+				{
+					return FString::Printf(TEXT("On demand state - %s [%i]"), *StateEnum->GetDisplayNameTextByValue(State).ToString(), ActorData.OnDemandInstanceIndex);
+				}
+
+				/** Otherwise we should just be part of a state */
+				return FString::Printf(TEXT("State - %s %1.2f"), *StateEnum->GetDisplayNameTextByValue(State).ToString(), ActorData.SignificanceValue);
+			}();
+
+			const FColor DebugColor = [&]()
+			{
+				const uint32 BlendInstanceIndex = ActorData.BlendInstanceIndex;
+				const uint32 DemandInstanceIndex = ActorData.OnDemandInstanceIndex;
+
+				/** Colors match debug material colors */				
+				if (ActorData.bBlending && BlendInstanceIndex != INDEX_NONE)
+				{
+					return FColor::Blue;
+				}
+				else if (ActorData.bRunningOnDemand && DemandInstanceIndex != INDEX_NONE)
+				{
+					return FColor::Red;
+				}
+
+				return FColor::Green;
+			}();
+
+#if UE_BUILD_DEVELOPMENT
+			/** Draw text above AI pawn's head */
+			DrawDebugString(GetWorld(), FVector(0.f, 0.f, 100.f), StateString, Actor, DebugColor, 0.016f, false);
+#endif
+			if (GAnimationSharingDebugging >= 2)
+			{
+				const FString OnScreenString = FString::Printf(TEXT("%s\n\tState %s [%i]\n\t%s\n\tBlending %i On-Demand %i"), *Actor->GetName(), *StateEnum->GetDisplayNameTextByValue(ActorData.CurrentState).ToString(), ActorData.PermutationIndex, *StateEnum->GetDisplayNameTextByValue(ActorData.PreviousState).ToString(), ActorData.bBlending, ActorData.bRunningOnDemand);
+
+				GEngine->AddOnScreenDebugMessage(1337, 1, FColor::White, OnScreenString);
+
+				const USkinnedMeshComponent* Component = PerComponentData[ActorData.ComponentIndices[0]].Component->MasterPoseComponent.Get();
+#if UE_BUILD_DEVELOPMENT
+				if (Component != nullptr)
+				{
+					DrawDebugLine(GetWorld(), Actor->GetActorLocation(), Component->GetComponentLocation(), FColor::Magenta);
+				}
+#endif
+			}
+
+			
+		}
+	}
+#endif
+}
+
+void UAnimSharingInstance::TickOnDemandInstances()
+{
+	SCOPE_CYCLE_COUNTER(STAT_AnimationSharing_UpdateOnDemands);
+	for (int32 InstanceIndex = 0; InstanceIndex < OnDemandInstances.Num(); ++ InstanceIndex)
+	{
+		FOnDemandInstance& Instance = OnDemandInstances[InstanceIndex];
+		checkf(Instance.bActive, TEXT("Container should be active at this point"));
+
+		// Mark on-demand component as in-use
+		SetComponentUsage(true, Instance.State, Instance.UsedPerStateComponentIndex);
+
+		const bool bShouldTick = DoAnyActorsRequireTicking(Instance);
+		if (bShouldTick)
+		{
+			// Mark component to tick
+			SetComponentTick(Instance.State, Instance.UsedPerStateComponentIndex);
+		}
+
+		// Check and see whether or not the animation has finished
+		if (Instance.EndTime <= WorldTime)
+		{
+			// Set in-use flag to false this should set the component to not tick during the next TickAnimationStates
+			SetComponentUsage(false, Instance.State, Instance.UsedPerStateComponentIndex);
+
+#if LOG_STATES 
+			UE_LOG(LogAnimationSharing, Log, TEXT("Finished on demand %s"), *StateEnum->GetDisplayNameTextByValue(Instance.State).ToString());
+#endif
+			auto SetActorState = [&](uint32 ActorIndex, uint8 NewState)
+			{
+				if (Instance.BlendToPermutationIndex != INDEX_NONE)
+				{
+					SetPermutationSlaveComponent(NewState, ActorIndex, Instance.BlendToPermutationIndex);
+				}
+				else
+				{
+					SetupSlaveComponent(NewState, ActorIndex);
+				}
+
+				// Set actor states
+				PerActorData[ActorIndex].PreviousState = PerActorData[ActorIndex].CurrentState;
+				PerActorData[ActorIndex].CurrentState = NewState;
+			};				
+			
+			
+			// Set the components to their current state animation
+			for (uint32 ActorIndex : Instance.ActorIndices)
+			{					
+				const uint32 CurrentState = PerActorData[ActorIndex].CurrentState;
+				// Return to the previous active animation state 
+				if (Instance.bReturnToPreviousState)
+				{
+					//for (uint32 ActorIndex : Instance.ActorIndices)
+					{
+						// Retrieve previous state for the actor 
+						const uint8 PreviousActorState = PerActorData[ActorIndex].PreviousState;
+						SetActorState(ActorIndex, PreviousActorState);
+#if LOG_STATES 
+						UE_LOG(LogAnimationSharing, Log, TEXT("Returning [%i] to %s"), ActorIndex, *StateEnum->GetDisplayNameTextByValue(PreviousActorState).ToString());
+#endif
+					}
+				}
+				else if (Instance.ForwardState != (uint8)INDEX_NONE)
+				{
+					// We could forward it to a different state at this point						
+					SetActorState(ActorIndex, Instance.ForwardState);
+#if LOG_STATES
+					UE_LOG(LogAnimationSharing, Log, TEXT("Forwarding [%i] to %s"), ActorIndex, *StateEnum->GetDisplayNameTextByValue(Instance.ForwardState).ToString());
+#endif						
+				}
+				// Only do this if the state is different than the current on-demand one
+				else if (CurrentState != Instance.State)
+				{
+					// If the new state is not an on-demand one and we are not currently blending, if we are blending the blend will set the final master component
+					if (!PerStateData[CurrentState].bIsOnDemand || !Instance.bBlendActive)
+					{
+						SetActorState(ActorIndex, CurrentState);
+
+						UAnimationSharingManager::SetDebugMaterialForActor(this, ActorIndex, 1);
+#if LOG_STATES 
+						UE_LOG(LogAnimationSharing, Log, TEXT("Setting [%i] to %s"), ActorIndex, *StateEnum->GetDisplayNameTextByValue(CurrentState).ToString());
+#endif
+					}
+				}
+				else
+				{
+					// Otherwise what do we do TODO
+#if LOG_STATES 
+					UE_LOG(LogAnimationSharing, Log, TEXT("TODO-ing [%i]"), ActorIndex);
+#endif
+				}					
+			}			
+
+			// Clear out data for each actor part of this instance
+			for (uint32 ActorIndex : Instance.ActorIndices)
+			{
+				const bool bPartOfOtherOnDemand = PerActorData[ActorIndex].OnDemandInstanceIndex != InstanceIndex;
+				//ensureMsgf(!bPartOfOtherOnDemand, TEXT("Actor on demand index differs from current instance"));
+
+				PerActorData[ActorIndex].OnDemandInstanceIndex = INDEX_NONE;
+				PerActorData[ActorIndex].bRunningOnDemand = false;			
+			}
+
+			// Remove this instance as it has finished work
+			RemoveOnDemandInstance(InstanceIndex);
+
+			// Decrement index so we don't skip the swapped instance
+			--InstanceIndex;
+		}
+		else if (!Instance.bBlendActive && Instance.StartBlendTime <= WorldTime)
+		{
+			for (uint32 ActorIndex : Instance.ActorIndices)
+			{
+				// Whether or not we can/should actually blend
+				const bool bShouldBlend = ScalabilitySettings->UseBlendTransitions.Default && PerActorData[ActorIndex].SignificanceValue >= ScalabilitySettings->BlendSignificanceValue.Default;
+
+				// Determine state to blend to
+				const uint8 BlendToState = [&]() -> uint8
+				{
+					if (bShouldBlend)
+					{
+						bool bShouldProcess;
+						const uint32 DeterminedState = DetermineStateForActor(ActorIndex, bShouldProcess);
+						const uint32 CurrentState = PerActorData[ActorIndex].CurrentState != DeterminedState ? DeterminedState : PerActorData[ActorIndex].CurrentState;
+
+						if (Instance.bReturnToPreviousState)
+						{
+							// Setup blend from on-demand animation into next state animation
+							return PerActorData[ActorIndex].PreviousState;
+						}
+						else if (Instance.ForwardState != (uint8)INDEX_NONE)
+						{
+							// Blend into the forward state 
+							return Instance.ForwardState;
+						}
+						else if (PerActorData[ActorIndex].CurrentState != Instance.State)
+						{
+							// Blend to the actor's current state
+							return PerActorData[ActorIndex].CurrentState;
+						}
+					}
+					return INDEX_NONE;					
+				}();
+
+				// Try to setup blending
+				if (BlendToState != (uint8)INDEX_NONE)
+				{
+					const uint32 BlendIndex = SetupBlendFromOnDemand(BlendToState, InstanceIndex, ActorIndex);
+
+					if (BlendIndex != INDEX_NONE)
+					{
+						// TODO what if two actors have a different state they are blending to? --> Store permutation index
+						Instance.BlendToPermutationIndex = BlendInstances[BlendIndex].ToPermutationIndex;
+#if LOG_STATES 
+						UE_LOG(LogAnimationSharing, Log, TEXT("Blending [%i] out from %s to %s"), ActorIndex, *StateEnum->GetDisplayNameTextByValue(Instance.State).ToString(), *StateEnum->GetDisplayNameTextByValue(BlendToState).ToString());
+#endif
+					}
+				}
+
+				// OR results, some actors could not be blending 
+				Instance.bBlendActive |= bShouldBlend;
+			}
+		}
+	}	
+}
+
+void UAnimSharingInstance::TickAdditiveInstances()
+{
+	SCOPE_CYCLE_COUNTER(STAT_AnimationSharing_UpdateAdditives);
+
+	for (int32 InstanceIndex = 0; InstanceIndex < AdditiveInstances.Num(); ++InstanceIndex)
+	{
+		FAdditiveInstance& Instance = AdditiveInstances[InstanceIndex];
+		if (Instance.bActive)
+		{
+			const float WorldTimeSeconds = GetWorld()->GetTimeSeconds();
+			if (WorldTimeSeconds >= Instance.EndTime)
+			{
+				// Finish
+				if (PerActorData.IsValidIndex(Instance.ActorIndex))
+				{
+					PerActorData[Instance.ActorIndex].bRunningAdditive = false;
+					PerActorData[Instance.ActorIndex].AdditiveInstanceIndex = INDEX_NONE;
+
+					// Set it to base component on top of the additive animation is playing
+					SetMasterComponentForActor(Instance.ActorIndex, Instance.AdditiveAnimationInstance->GetBaseComponent());
+				}
+				FreeAdditiveInstance(Instance.AdditiveAnimationInstance);
+				RemoveAdditiveInstance(InstanceIndex);
+				--InstanceIndex;
+			}
+		}
+		else
+		{
+			Instance.bActive = true;
+			Instance.AdditiveAnimationInstance->Start();
+			if (Instance.ActorIndex != INDEX_NONE)
+			{				
+				SetMasterComponentForActor(Instance.ActorIndex, Instance.AdditiveAnimationInstance->GetComponent());
+			}
+		}
+	}
+}
+
+void UAnimSharingInstance::TickActorStates()
+{
+	SCOPE_CYCLE_COUNTER(STAT_AnimationSharing_TickActorStates);
+	/** Tick each registered actor's state */
+	for (int32 ActorIndex = 0; ActorIndex < RegisteredActors.Num(); ++ActorIndex)
+	{
+		/** Ensure Actor is still available */
+		const AActor* Actor = RegisteredActors[ActorIndex];
+		if (Actor)
+		{
+			FPerActorData& ActorData = PerActorData[ActorIndex];
+			checkf(ActorData.ComponentIndices.Num(), TEXT("Registered Actor without SkeletalMeshComponents"));
+
+			// Update actor and component visibility
+			ActorData.bRequiresTick = ActorData.SignificanceValue >= ScalabilitySettings->TickSignificanceValue.Default;
+			for (int32 ComponentIndex : ActorData.ComponentIndices)
+			{
+				if (PerComponentData[ComponentIndex].Component->LastRenderTime > (WorldTime - 1.f))
+				{
+					PerComponentData[ComponentIndex].Component->bRecentlyRendered = true;
+					ActorData.bRequiresTick = true;
+				}
+			}
+
+			// Determine current state for Actor
+			uint8& PreviousState = ActorData.CurrentState;
+			bool bShouldProcess = false;
+			const uint8 CurrentState = DetermineStateForActor(ActorIndex, bShouldProcess);
+
+			// Determine whether we should blend according to the scalability settings
+			const bool bShouldBlend = ScalabilitySettings->UseBlendTransitions.Default && ActorData.SignificanceValue >= ScalabilitySettings->BlendSignificanceValue.Default;
+
+			/** If the state is different we need to change animations and setup a transition */
+			if (CurrentState != PreviousState)
+			{
+				/** When we are currently running an on-demand state we do not want as state change to impact the current animation */
+				const bool bShouldNotProcess = ActorData.bRunningOnDemand && !PerStateData[CurrentState].bIsOnDemand;
+
+				auto UpdateState = [&]()
+				{
+#if LOG_STATES 
+					UE_LOG(LogAnimationSharing, Log, TEXT("Setting %i state to %i previous %i | %i"), ActorIndex, CurrentState, PreviousState, ActorData.PermutationIndex);
+#endif
+					ActorData.PreviousState = PreviousState;
+					ActorData.CurrentState = CurrentState;
+				};
+
+				/** If the processor explicitly outputs that the change in state should not impact behavior, just change state and do nothing */
+				if (!bShouldProcess || bShouldNotProcess)
+				{
+					UpdateState();
+#if LOG_STATES 
+					UE_LOG(LogAnimationSharing, Log, TEXT("Changing state to %s from %s while running on demand %i"), *StateEnum->GetDisplayNameTextByValue(CurrentState).ToString(), *StateEnum->GetDisplayNameTextByValue(ActorData.PreviousState).ToString(), ActorIndex);
+#endif
+				}
+				/** Play additive animation only if actor isn't already playing one */
+				else if (PerStateData[CurrentState].bIsAdditive && !ActorData.bRunningAdditive)
+				{					
+					const uint32 AdditiveInstanceIndex = SetupAdditiveInstance(CurrentState, PreviousState, ActorData.PermutationIndex);
+					if (AdditiveInstanceIndex != INDEX_NONE)
+					{
+						ActorData.bRunningAdditive = true;
+						ActorData.AdditiveInstanceIndex = AdditiveInstanceIndex;
+						AdditiveInstances[AdditiveInstanceIndex].ActorIndex = ActorIndex;
+					}
+				}
+				/** If we are _already_ running an on-demand instance and the new state is also an on-demand we'll have to blend the new state in*/
+				else if (PerStateData[CurrentState].bIsOnDemand)					
+				{
+					/** If the new state is different than the currently running on-demand state, this could happen if we previously only updated the state and not processed it */
+					const bool bSetupInstance = (!ActorData.bRunningOnDemand || (ActorData.bRunningOnDemand && OnDemandInstances[ActorData.OnDemandInstanceIndex].State != CurrentState));
+					const uint32 OnDemandIndex = bSetupInstance ? SetupOnDemandInstance(CurrentState) : INDEX_NONE;
+
+					if (OnDemandIndex != INDEX_NONE)
+					{
+						// Make sure we end any current blends
+						RemoveFromCurrentBlend(ActorIndex);
+						RemoveFromCurrentOnDemand(ActorIndex);
+
+						bool bShouldSwitch = true;
+						if (bShouldBlend && !FMath::IsNearlyZero(PerStateData[CurrentState].BlendTime))
+						{
+							if (ActorData.bRunningOnDemand)
+							{
+								/** Setup a blend between the current and a new instance*/
+								const uint32 BlendInstanceIndex = SetupBlendBetweenOnDemands(ActorData.OnDemandInstanceIndex, OnDemandIndex, ActorIndex);
+								ActorData.BlendInstanceIndex = BlendInstanceIndex;
+							}
+							else
+							{
+								/** Setup a blend to an on-demand state/instance */
+								const uint32 BlendInstanceIndex = SetupBlendToOnDemand(PreviousState, OnDemandIndex, ActorIndex);
+								ActorData.BlendInstanceIndex = BlendInstanceIndex;
+							}
+
+							/** Blend was not succesfully set up so switch anyway */
+							bShouldSwitch = (ActorData.BlendInstanceIndex == INDEX_NONE);
+						}
+
+						if (bShouldSwitch)
+						{
+							/** Not blending so just switch to other on-demand instance */
+							SwitchBetweenOnDemands(ActorData.OnDemandInstanceIndex, OnDemandIndex, ActorIndex);
+						}
+
+						/** Add the current actor to the on-demand instance*/
+						OnDemandInstances[OnDemandIndex].ActorIndices.Add(ActorIndex);
+						/** Also change actor data accordingly*/
+						ActorData.OnDemandInstanceIndex = OnDemandIndex;
+						ActorData.bRunningOnDemand = true;
+
+						UpdateState();
+					}
+				}
+				/** Otherwise blend towards the new shared state */
+				else
+				{
+					/** If actor is within blending distance setup/reuse a blend instance*/
+					bool bShouldSwitch = true;
+					if (bShouldBlend)
+					{
+						const uint32 BlendInstanceIndex = SetupBlend(PreviousState, CurrentState, ActorIndex);
+						ActorData.BlendInstanceIndex = BlendInstanceIndex;
+						/** Blend was not succesfully set up so switch anyway */
+						bShouldSwitch = (ActorData.BlendInstanceIndex == INDEX_NONE);
+#if LOG_STATES 
+						UE_LOG(LogAnimationSharing, Log, TEXT("Changing state to %s from %s with blend %i"), *StateEnum->GetDisplayNameTextByValue(CurrentState).ToString(), *StateEnum->GetDisplayNameTextByValue(PreviousState).ToString(), ActorIndex);
+#endif
+					}
+					/** Otherwise just switch it to the new state */
+					if (bShouldSwitch)
+					{
+						SetupSlaveComponent(CurrentState, ActorIndex);
+#if LOG_STATES 
+						UE_LOG(LogAnimationSharing, Log, TEXT("Changing state to %s from %s %i"), *StateEnum->GetDisplayNameTextByValue(CurrentState).ToString(), *StateEnum->GetDisplayNameTextByValue(PreviousState).ToString(), ActorIndex);
+#endif
+					}
+
+					UpdateState();
+				}
+			}
+			/** Flag the currently master component as in-use */
+			else if (!ActorData.bRunningOnDemand && !ActorData.bBlending)
+			{
+#if LOG_STATES 
+				if (!PerStateData[ActorData.CurrentState].Components.IsValidIndex(ActorData.PermutationIndex))
+				{
+					UE_LOG(LogAnimationSharing, Log, TEXT("Invalid permutation for actor %i is out of range of %i for state %s by actor %i"), ActorData.PermutationIndex, PerStateData[ActorData.CurrentState].Components.Num(), *StateEnum->GetDisplayNameTextByValue(ActorData.CurrentState).ToString(), ActorIndex);
+				}
+				else if (!PerStateData[ActorData.CurrentState].Components[ActorData.PermutationIndex]->IsComponentTickEnabled())
+				{
+					UE_LOG(LogAnimationSharing, Log, TEXT("Component not active %i for state %s by actor %i"), ActorData.PermutationIndex, *StateEnum->GetDisplayNameTextByValue(ActorData.CurrentState).ToString(), ActorIndex);
+				}
+#endif 
+
+				SetComponentUsage(true, ActorData.CurrentState, ActorData.PermutationIndex);
+			}
+			
+			// Propagate visibility to master component
+			if (ActorData.bRequiresTick)
+			{
+				SetComponentTick(ActorData.CurrentState, ActorData.PermutationIndex);
+			}
+		}
+	}
+}
+
+void UAnimSharingInstance::RemoveFromCurrentBlend(int32 ActorIndex)
+{
+	if (PerActorData[ActorIndex].bBlending && PerActorData[ActorIndex].BlendInstanceIndex != INDEX_NONE && BlendInstances.IsValidIndex(PerActorData[ActorIndex].BlendInstanceIndex))
+	{
+		FBlendInstance& OldBlendInstance = BlendInstances[PerActorData[ActorIndex].BlendInstanceIndex];
+		SetMasterComponentForActor(ActorIndex, OldBlendInstance.TransitionBlendInstance->GetToComponent());
+		OldBlendInstance.ActorIndices.Remove(ActorIndex);
+		PerActorData[ActorIndex].BlendInstanceIndex = INDEX_NONE;
+	}
+}
+
+void UAnimSharingInstance::RemoveFromCurrentOnDemand(int32 ActorIndex)
+{
+	if (PerActorData[ActorIndex].bRunningOnDemand && PerActorData[ActorIndex].OnDemandInstanceIndex != INDEX_NONE && OnDemandInstances.IsValidIndex(PerActorData[ActorIndex].OnDemandInstanceIndex))
+	{
+		OnDemandInstances[PerActorData[ActorIndex].OnDemandInstanceIndex].ActorIndices.Remove(ActorIndex);
+	}
+}
+
+void UAnimSharingInstance::TickBlendInstances()
+{
+	SCOPE_CYCLE_COUNTER(STAT_AnimationSharing_UpdateBlends);
+	for (int32 InstanceIndex = 0; InstanceIndex < BlendInstances.Num(); ++InstanceIndex)
+	{
+		FBlendInstance& Instance = BlendInstances[InstanceIndex];
+		checkf(Instance.bActive, TEXT("Blends should be active at this point"));
+
+		/** Check whether or not the blend has ended */
+		if (Instance.EndTime <= WorldTime)
+		{
+#if LOG_STATES
+			UE_LOG(LogAnimationSharing, Log, TEXT("Finished blend %s from %s"), *StateEnum->GetDisplayNameTextByValue(Instance.StateTo).ToString(), *StateEnum->GetDisplayNameTextByValue(Instance.StateFrom).ToString());
+#endif
+
+			// Finish blend into unique animation, need to just set it to use the correct component
+			const bool bToStateIsOnDemand = PerStateData[Instance.StateTo].bIsOnDemand;
+			const bool bFromStateIsOnDemand = PerStateData[Instance.StateFrom].bIsOnDemand;
+
+			// If we were blending to an on-demand state we need to set the on-demand component as the new master component
+			if (bToStateIsOnDemand)
+			{
+				for (uint32 ActorIndex : Instance.ActorIndices)
+				{
+					SetMasterComponentForActor(ActorIndex, Instance.TransitionBlendInstance->GetToComponent());
+					PerActorData[ActorIndex].PermutationIndex = 0;
+#if LOG_STATES
+					UE_LOG(LogAnimationSharing, Log, TEXT("Setting %i to on-demand component %i"), ActorIndex, Instance.ToOnDemandInstanceIndex);
+#endif
+
+					for (uint32 ComponentIndex : PerActorData[ActorIndex].ComponentIndices)
+					{
+						UAnimationSharingManager::SetDebugMaterial(PerComponentData[ComponentIndex].Component, 0);
+					}
+				}
+			}
+			/** Otherwise if the state we were blending from was not on-demand we set the new state component as the new master component,
+				if we are blending from an on-demand state FOnDemandInstance with set the correct master component when it finishes	*/
+			else if (!bFromStateIsOnDemand)
+			{				
+				for (uint32 ActorIndex : Instance.ActorIndices)
+				{
+					if (PerActorData[ActorIndex].CurrentState == Instance.StateTo)
+					{
+#if LOG_STATES 
+						UE_LOG(LogAnimationSharing, Log, TEXT("Setting %i to state %i | %i"), ActorIndex, Instance.StateTo, Instance.ToPermutationIndex);
+#endif
+						SetPermutationSlaveComponent(Instance.StateTo, ActorIndex, Instance.ToPermutationIndex);
+#if !UE_BUILD_SHIPPING
+						for (uint32 ComponentIndex : PerActorData[ActorIndex].ComponentIndices)
+						{
+							UAnimationSharingManager::SetDebugMaterial(PerComponentData[ComponentIndex].Component, 1);
+						}
+#endif
+					}
+				
+				}
+			}			
+
+			// Free up the used blend actor
+			FreeBlendInstance(Instance.TransitionBlendInstance);
+
+			// Clear flags and index on the actor data as the blend has finished
+			for (uint32 ActorIndex : Instance.ActorIndices)
+			{
+				PerActorData[ActorIndex].BlendInstanceIndex = INDEX_NONE;
+				PerActorData[ActorIndex].bBlending = 0;
+			}
+
+			// Remove this blend instance as it has finished
+			RemoveBlendInstance(InstanceIndex);
+			--InstanceIndex;
+		}
+		else
+		{
+			// Check whether or not the blend has started, if not set up the actors as slaves at this point
+			if (!Instance.bBlendStarted)
+			{
+				for (uint32 ActorIndex : Instance.ActorIndices)
+				{
+					SetMasterComponentForActor(ActorIndex, Instance.TransitionBlendInstance->GetComponent());			
+
+					for (uint32 ComponentIndex : PerActorData[ActorIndex].ComponentIndices)
+					{
+						UAnimationSharingManager::SetDebugMaterial(PerComponentData[ComponentIndex].Component, 2);
+					}
+				}
+
+				Instance.bBlendStarted = true;
+			}
+
+			const bool bShouldTick = DoAnyActorsRequireTicking(Instance);
+
+			if (!PerStateData[Instance.StateFrom].bIsOnDemand)
+			{
+				SetComponentUsage(true, Instance.StateFrom, Instance.FromPermutationIndex);
+				if (bShouldTick)
+				{
+					SetComponentTick(Instance.StateFrom, Instance.FromPermutationIndex);
+				}
+			}
+
+			if (!PerStateData[Instance.StateTo].bIsOnDemand)
+			{
+				SetComponentUsage(true, Instance.StateTo, Instance.ToPermutationIndex);
+				if (bShouldTick)
+				{
+					SetComponentTick(Instance.StateTo, Instance.ToPermutationIndex);
+				}
+			}
+		}
+	}
+}
+
+void UAnimSharingInstance::TickAnimationStates()
+{
+	for (FPerStateData& StateData : PerStateData)
+	{
+		for (int32 Index = 0; Index < StateData.Components.Num(); ++Index)
+		{
+			const bool bPreviousState = StateData.PreviousInUseComponentFrameBits[Index];
+			const bool bCurrentState = StateData.InUseComponentFrameBits[Index];
+
+			const bool bShouldTick = StateData.SlaveTickRequiredFrameBits[Index];
+
+			if (bCurrentState != bPreviousState)
+			{
+				if (bCurrentState)
+				{
+					// Turn on
+					UAnimationSharingManager::SetDebugMaterial(StateData.Components[Index], 1);
+					StateData.Components[Index]->SetComponentTickEnabled(true);
+				}
+				else
+				{
+					// Turn off
+					UAnimationSharingManager::SetDebugMaterial(StateData.Components[Index], 0);
+					StateData.Components[Index]->SetComponentTickEnabled(false);
+				}
+			}
+			else if (!bCurrentState && StateData.Components[Index]->IsComponentTickEnabled())
+			{
+				// Turn off
+				UAnimationSharingManager::SetDebugMaterial(StateData.Components[Index], 0);
+				StateData.Components[Index]->SetComponentTickEnabled(false);
+			}
+
+			StateData.Components[Index]->bRecentlyRendered = bShouldTick;
+			StateData.Components[Index]->VisibilityBasedAnimTickOption = bShouldTick ? EVisibilityBasedAnimTickOption::AlwaysTickPoseAndRefreshBones : EVisibilityBasedAnimTickOption::OnlyTickPoseWhenRendered;
+		}
+
+		// Set previous to current and reset current bits
+		StateData.PreviousInUseComponentFrameBits = StateData.InUseComponentFrameBits;
+		StateData.InUseComponentFrameBits.Init(false, StateData.PreviousInUseComponentFrameBits.Num());
+		StateData.SlaveTickRequiredFrameBits.Init(false, StateData.SlaveTickRequiredFrameBits.Num());
+
+		/** Reset on demand index for next frame */
+		StateData.CurrentFrameOnDemandIndex = INDEX_NONE;
+	}	
+}
+
+void UAnimSharingInstance::SetComponentUsage(bool bUsage, uint8 StateIndex, uint32 ComponentIndex)
+{
+	// TODO component index should always be valid
+#if LOG_STATES 
+	if (!PerStateData[StateIndex].InUseComponentFrameBits.IsValidIndex(ComponentIndex))
+	{
+		UE_LOG(LogAnimationSharing, Log, TEXT("Invalid set component usage %i is out of range of %i for state %s by component %i"), ComponentIndex, PerStateData[StateIndex].Components.Num(), *StateEnum->GetDisplayNameTextByValue(StateIndex).ToString(), ComponentIndex);
+	}
+#endif
+
+	if (PerStateData.IsValidIndex(StateIndex) && PerStateData[StateIndex].InUseComponentFrameBits.IsValidIndex(ComponentIndex))
+	{
+		PerStateData[StateIndex].InUseComponentFrameBits[ComponentIndex] = bUsage;
+	}
+}
+
+void UAnimSharingInstance::SetComponentTick(uint8 StateIndex, uint32 ComponentIndex)
+{
+	if (PerStateData[StateIndex].SlaveTickRequiredFrameBits.IsValidIndex(ComponentIndex))
+	{
+		PerStateData[StateIndex].SlaveTickRequiredFrameBits[ComponentIndex] = true;
+	}
+}
+
+void UAnimSharingInstance::FreeBlendInstance(FTransitionBlendInstance* Instance)
+{
+	Instance->Stop();
+	BlendInstanceStack.FreeInstance(Instance);
+}
+
+void UAnimSharingInstance::FreeAdditiveInstance(FAdditiveAnimationInstance* Instance)
+{
+	Instance->Stop();
+	AdditiveInstanceStack.FreeInstance(Instance);
+}
+
+void UAnimSharingInstance::SetMasterComponentForActor(uint32 ActorIndex, USkeletalMeshComponent* Component)
+{
+	// Always ensure the component is ticking
+	if (Component)
+	{
+		Component->SetComponentTickEnabled(true);
+	}
+
+	const FPerActorData& ActorData = PerActorData[ActorIndex];
+	// Do not update the component of the additive actor itself, otherwise update the base component
+	if (ActorData.bRunningAdditive && AdditiveInstances.IsValidIndex(ActorData.AdditiveInstanceIndex) && AdditiveInstances[ActorData.AdditiveInstanceIndex].AdditiveAnimationInstance->GetComponent() != Component)
+	{		
+		AdditiveInstances[ActorData.AdditiveInstanceIndex].BaseComponent = Component;
+		AdditiveInstances[ActorData.AdditiveInstanceIndex].AdditiveAnimationInstance->UpdateBaseComponent(Component);
+
+		return;
+	}
+	
+	for (uint32 ComponentIndex : ActorData.ComponentIndices)
+	{
+		PerComponentData[ComponentIndex].Component->SetMasterPoseComponent(Component, true);
+	}
+}
+
+void UAnimSharingInstance::SetupSlaveComponent(uint8 CurrentState, uint32 ActorIndex)
+{
+	const FPerStateData& StateData = PerStateData[CurrentState];
+	if (!StateData.bIsOnDemand)
+	{
+		const uint32 PermutationIndex = DeterminePermutationIndex(ActorIndex, CurrentState);
+		SetPermutationSlaveComponent(CurrentState, ActorIndex, PermutationIndex);
+	}
+	else
+	{
+		const uint32 OnDemandInstanceIndex = SetupOnDemandInstance(CurrentState);
+
+		if (OnDemandInstanceIndex != INDEX_NONE)
+		{
+			USkeletalMeshComponent* MasterComponent = StateData.Components[OnDemandInstances[OnDemandInstanceIndex].UsedPerStateComponentIndex];
+			SetMasterComponentForActor(ActorIndex, MasterComponent);
+			OnDemandInstances[OnDemandInstanceIndex].ActorIndices.Add(ActorIndex);
+
+			PerActorData[ActorIndex].OnDemandInstanceIndex = OnDemandInstanceIndex;
+			PerActorData[ActorIndex].bRunningOnDemand = true;
+
+			// TODO do we need to reset
+			PerActorData[ActorIndex].PermutationIndex = 0;				
+		}
+	}
+}
+
+void UAnimSharingInstance::SetPermutationSlaveComponent(uint8 StateIndex, uint32 ActorIndex, uint32 PermutationIndex)
+{
+	const FPerStateData& StateData = PerStateData[StateIndex];
+
+	// TODO Min should not be needed if PermutationIndex is always valid
+	PermutationIndex = FMath::Min((uint32)StateData.Components.Num() - 1, PermutationIndex);
+#if LOG_STATES 
+	if (!StateData.Components.IsValidIndex(PermutationIndex))
+	{
+		UE_LOG(LogAnimationSharing, Log, TEXT("Invalid set component usage %i is out of range of %i for state %s by actor %i"), PermutationIndex, StateData.Components.Num(), *StateEnum->GetDisplayNameTextByValue(StateIndex).ToString(), ActorIndex);
+	}
+#endif
+
+	SetMasterComponentForActor(ActorIndex, StateData.Components[PermutationIndex]);
+	PerActorData[ActorIndex].PermutationIndex = PermutationIndex;
+	UAnimationSharingManager::SetDebugMaterial(StateData.Components[PermutationIndex], 1);
+}
+
+uint32 UAnimSharingInstance::DeterminePermutationIndex(uint32 ActorIndex, uint8 State) const
+{
+	const FPerStateData& StateData = PerStateData[State];
+	const TArray<USkeletalMeshComponent*>& Components = StateData.Components;
+
+	// This can grow to be more intricate to take into account surrounding actors?
+	const uint32 PermutationIndex = FMath::RandHelper(Components.Num());
+	checkf(Components.IsValidIndex(PermutationIndex), TEXT("Not enough MasterComponents initialised!"));
+
+	return PermutationIndex;
+}
+
+uint32 UAnimSharingInstance::SetupBlend(uint8 FromState, uint8 ToState, uint32 ActorIndex)
+{
+	const bool bConcurrentBlendsReached = !BlendInstanceStack.InstanceAvailable();
+	const bool bOnDemand = PerStateData[ToState].bIsOnDemand;
+
+	uint32 BlendInstanceIndex = INDEX_NONE;
+	if (!bConcurrentBlendsReached)
+	{
+		BlendInstanceIndex = BlendInstances.IndexOfByPredicate([&](const FBlendInstance& Instance)
+		{			
+			return (!Instance.bActive &&				// The instance should not have started yet
+				Instance.StateFrom == FromState &&		// It should be blending from the same state
+				Instance.StateTo == ToState &&			// It should be blending to the same state
+				Instance.bOnDemand == bOnDemand &&		// It should match whether or not it is an on-demand state QQQ is this needed?
+				Instance.FromPermutationIndex == PerActorData[ActorIndex].PermutationIndex); // It should be blending from the same permutation inside of the state 
+		});
+
+		FBlendInstance* BlendInstance = BlendInstanceIndex != INDEX_NONE ? &BlendInstances[BlendInstanceIndex] : nullptr;
+
+		if (!BlendInstance)
+		{
+			BlendInstance = &BlendInstances.AddDefaulted_GetRef();
+			BlendInstanceIndex = BlendInstances.Num() - 1;
+			BlendInstance->bActive = false;
+			BlendInstance->FromOnDemandInstanceIndex = BlendInstance->ToOnDemandInstanceIndex = INDEX_NONE;
+			BlendInstance->StateFrom = FromState;
+			BlendInstance->StateTo = ToState;
+			BlendInstance->BlendTime = CalculateBlendTime( ToState);
+			BlendInstance->bOnDemand = bOnDemand;
+			BlendInstance->EndTime = GetWorld()->GetTimeSeconds() + BlendInstance->BlendTime;
+			BlendInstance->TransitionBlendInstance = BlendInstanceStack.GetInstance();
+
+			BlendInstance->TransitionBlendInstance->GetComponent()->SetComponentTickEnabled(true);
+
+			// Setup permutation indices to and from we are blending
+			BlendInstance->FromPermutationIndex = PerActorData[ActorIndex].PermutationIndex;
+			BlendInstance->ToPermutationIndex = DeterminePermutationIndex( ActorIndex, ToState);
+		}
+
+		checkf(BlendInstance, TEXT("Unable to create blendcontainer"));
+
+		BlendInstance->ActorIndices.Add(ActorIndex);
+		PerActorData[ActorIndex].bBlending = true;
+	}
+
+	return BlendInstanceIndex;
+}
+
+uint32 UAnimSharingInstance::SetupBlendFromOnDemand(uint8 ToState, uint32 OnDemandInstanceIndex, uint32 ActorIndex)
+{
+	const uint8 FromState = OnDemandInstances[OnDemandInstanceIndex].State;
+	const uint32 BlendInstanceIndex = SetupBlend(FromState, ToState, ActorIndex);
+
+	if (BlendInstanceIndex != INDEX_NONE)
+	{
+		BlendInstances[BlendInstanceIndex].FromOnDemandInstanceIndex = OnDemandInstanceIndex;
+	}
+
+	return BlendInstanceIndex;
+}
+
+uint32 UAnimSharingInstance::SetupBlendBetweenOnDemands(uint8 FromOnDemandInstanceIndex, uint32 ToOnDemandInstanceIndex, uint32 ActorIndex)
+{
+	const uint8 FromState = OnDemandInstances[FromOnDemandInstanceIndex].State;
+	const uint8 ToState = OnDemandInstances[ToOnDemandInstanceIndex].State;
+	const uint32 BlendInstanceIndex = SetupBlend(FromState, ToState, ActorIndex);
+
+	if (BlendInstanceIndex != INDEX_NONE)
+	{
+		BlendInstances[BlendInstanceIndex].FromOnDemandInstanceIndex = FromOnDemandInstanceIndex;
+		BlendInstances[BlendInstanceIndex].ToOnDemandInstanceIndex = ToOnDemandInstanceIndex;
+	}
+
+	return BlendInstanceIndex;
+}
+
+uint32 UAnimSharingInstance::SetupBlendToOnDemand(uint8 FromState, uint32 ToOnDemandInstanceIndex, uint32 ActorIndex)
+{
+	const uint8 ToState = OnDemandInstances[ToOnDemandInstanceIndex].State;
+	const uint32 BlendInstanceIndex = SetupBlend(FromState, ToState, ActorIndex);
+
+	if (BlendInstanceIndex != INDEX_NONE)
+	{
+		BlendInstances[BlendInstanceIndex].ToOnDemandInstanceIndex = ToOnDemandInstanceIndex;
+	}
+
+	return BlendInstanceIndex;
+}
+
+void UAnimSharingInstance::SwitchBetweenOnDemands(uint32 FromOnDemandInstanceIndex, uint32 ToOnDemandInstanceIndex, uint32 ActorIndex)
+{
+	/** Remove this actor from the currently running on-demand instance */
+	if (FromOnDemandInstanceIndex != INDEX_NONE)
+	{
+		OnDemandInstances[FromOnDemandInstanceIndex].ActorIndices.Remove(ActorIndex);
+	}
+
+	const FOnDemandInstance& Instance = OnDemandInstances[ToOnDemandInstanceIndex];
+	const uint32 ComponentIndex = Instance.UsedPerStateComponentIndex;
+	const uint32 StateIndex = Instance.State;
+	PerActorData[ActorIndex].PermutationIndex = 0;
+	SetMasterComponentForActor(ActorIndex, PerStateData[StateIndex].Components[ComponentIndex]);
+}
+
+uint32 UAnimSharingInstance::SetupOnDemandInstance(uint8 StateIndex)
+{
+	uint32 InstanceIndex = INDEX_NONE;
+
+	FPerStateData& StateData = PerStateData[StateIndex];
+	if (StateData.CurrentFrameOnDemandIndex != INDEX_NONE && OnDemandInstances.IsValidIndex(StateData.CurrentFrameOnDemandIndex))
+	{
+		InstanceIndex = StateData.CurrentFrameOnDemandIndex;
+	}
+	else
+	{
+		// Otherwise we'll need to kick one of right now so try and set one up		
+		if (StateData.Components.Num())
+		{
+			const uint32 AvailableIndex = StateData.InUseComponentFrameBits.FindAndSetFirstZeroBit();
+			
+			if (AvailableIndex != INDEX_NONE)
+			{
+				FOnDemandInstance& Instance = OnDemandInstances.AddDefaulted_GetRef();
+				InstanceIndex = OnDemandInstances.Num() - 1;
+				StateData.CurrentFrameOnDemandIndex = InstanceIndex;
+
+				Instance.bActive = 0;
+				Instance.bBlendActive = 0;
+				Instance.State = StateIndex;
+				Instance.ForwardState = StateData.bShouldForwardToState ? StateData.ForwardStateValue : INDEX_NONE;
+				Instance.UsedPerStateComponentIndex = AvailableIndex;
+				Instance.bReturnToPreviousState = StateData.bReturnToPreviousState;
+				Instance.StartTime = 0.f;
+				Instance.BlendToPermutationIndex = INDEX_NONE;
+
+				const float WorldTimeSeconds = GetWorld()->GetTimeSeconds();
+				Instance.EndTime = WorldTimeSeconds + StateData.AnimationLengths[AvailableIndex];
+				Instance.StartBlendTime = Instance.EndTime - CalculateBlendTime(StateIndex);
+
+				USkeletalMeshComponent* FreeComponent = StateData.Components[AvailableIndex];
+
+				UAnimationSharingManager::SetDebugMaterial(FreeComponent, 1);
+
+				FreeComponent->SetComponentTickEnabled(true);
+				FreeComponent->SetPosition(0.f, false);
+				FreeComponent->Play(false);
+#if LOG_STATES 
+				UE_LOG(LogAnimationSharing, Log, TEXT("Setup on demand state %s"), *StateEnum->GetDisplayNameTextByValue(StateIndex).ToString());
+#endif
+			}
+			else
+			{
+				// Next resort
+				const float MaxStartTime = WorldTime - PerStateData[StateIndex].WiggleTime;
+				float WiggleStartTime = TNumericLimits<float>::Max();
+				float NonWiggleStartTime = TNumericLimits<float>::Max();
+				int32 WiggleIndex = INDEX_NONE;
+				int32 NonWiggleIndex = INDEX_NONE;
+				for (int32 RunningInstanceIndex = 0; RunningInstanceIndex < OnDemandInstances.Num(); ++RunningInstanceIndex)
+				{
+					FOnDemandInstance& Instance = OnDemandInstances[RunningInstanceIndex];
+
+					if (Instance.State == StateIndex)
+					{
+						if (Instance.StartTime <= MaxStartTime && Instance.StartTime < WiggleStartTime)
+						{
+							WiggleStartTime = Instance.StartTime;
+							WiggleIndex = RunningInstanceIndex;
+						}
+						else if (Instance.StartTime < NonWiggleStartTime)
+						{
+							NonWiggleStartTime = Instance.StartTime;
+							NonWiggleIndex = RunningInstanceIndex;							
+						}
+					}
+				}
+
+				// Snap to on demand instance that has started last within the number of wiggle frames
+				if (WiggleIndex != INDEX_NONE)
+				{
+					InstanceIndex = WiggleIndex;
+				}
+				// Snap to closest on demand instance outside of the number of wiggle frames
+				else if (NonWiggleIndex != INDEX_NONE)
+				{
+					InstanceIndex = NonWiggleIndex;
+				}
+				else
+				{
+					// No instances available and none actually currently running this state, should probably up the number of available concurrent on demand instances at this point
+					UE_LOG(LogAnimationSharing, Warning, TEXT("No more on demand components available"));
+				}
+			}
+		}
+	}
+
+	return InstanceIndex;
+}
+
+uint32 UAnimSharingInstance::SetupAdditiveInstance(uint8 StateIndex, uint8 FromState, uint8 StateComponentIndex)
+{
+	uint32 InstanceIndex = INDEX_NONE;
+
+	const FPerStateData& StateData = PerStateData[StateIndex];
+	if (AdditiveInstanceStack.InstanceAvailable())
+	{
+		FAdditiveAnimationInstance* AnimationInstance = AdditiveInstanceStack.GetInstance();
+		FAdditiveInstance& Instance = AdditiveInstances.AddDefaulted_GetRef();
+		Instance.bActive = false;
+		Instance.AdditiveAnimationInstance = AnimationInstance;
+		Instance.BaseComponent = PerStateData[FromState].Components[StateComponentIndex];
+		const float WorldTimeSeconds = GetWorld()->GetTimeSeconds();
+		Instance.EndTime = WorldTimeSeconds + StateData.AdditiveAnimationSequence->SequenceLength;
+		Instance.State = StateIndex;
+
+		InstanceIndex = AdditiveInstances.Num() - 1;
+		AnimationInstance->Setup(Instance.BaseComponent, StateData.AdditiveAnimationSequence);
+	}
+
+	return InstanceIndex;
+}
+
+void UAnimSharingInstance::KickoffInstances()
+{
+	SCOPE_CYCLE_COUNTER(STAT_AnimationSharing_KickoffInstances);
+	for (FBlendInstance& BlendInstance : BlendInstances)
+	{
+		if (!BlendInstance.bActive)
+		{
+			BlendInstance.bBlendStarted = false;
+
+			FString ActorIndicesString;
+			for (uint32 ActorIndex : BlendInstance.ActorIndices)
+			{
+				if (ActorIndex != BlendInstance.ActorIndices.Last())
+				{
+					ActorIndicesString += FString::Printf(TEXT("%i, "), ActorIndex);
+				}
+				else
+				{
+					ActorIndicesString += FString::Printf(TEXT("%i"), ActorIndex);
+				}
+			}
+#if LOG_STATES 
+			UE_LOG(LogAnimationSharing, Log, TEXT("Starting blend from %s to %s [%s]"), *StateEnum->GetDisplayNameTextByValue(BlendInstance.StateFrom).ToString(), *StateEnum->GetDisplayNameTextByValue(BlendInstance.StateTo).ToString(), *ActorIndicesString);
+#endif
+
+			// TODO should be able to assume permutation indices are valid here
+			BlendInstance.FromPermutationIndex = FMath::Min((uint32)PerStateData[BlendInstance.StateFrom].Components.Num() - 1, BlendInstance.FromPermutationIndex);
+			BlendInstance.ToPermutationIndex = FMath::Min((uint32)PerStateData[BlendInstance.StateTo].Components.Num() - 1, BlendInstance.ToPermutationIndex);
+
+			USkeletalMeshComponent* From = PerStateData[BlendInstance.StateFrom].Components[BlendInstance.FromPermutationIndex];
+			USkeletalMeshComponent* To = PerStateData[BlendInstance.StateTo].Components[BlendInstance.ToPermutationIndex];
+
+			if (PerStateData[BlendInstance.StateTo].bIsOnDemand && (BlendInstance.ToOnDemandInstanceIndex != INDEX_NONE))
+			{
+				To = PerStateData[BlendInstance.StateTo].Components[OnDemandInstances[BlendInstance.ToOnDemandInstanceIndex].UsedPerStateComponentIndex];		
+			}
+
+			if (PerStateData[BlendInstance.StateFrom].bIsOnDemand && (BlendInstance.FromOnDemandInstanceIndex != INDEX_NONE))
+			{
+				const uint32 UsedComponentIndex = OnDemandInstances[BlendInstance.FromOnDemandInstanceIndex].UsedPerStateComponentIndex;
+				From = PerStateData[BlendInstance.StateFrom].Components[UsedComponentIndex];
+			}
+
+			for (uint32 ActorIndex : BlendInstance.ActorIndices)
+			{
+				PerActorData[ActorIndex].PermutationIndex = BlendInstance.ToPermutationIndex;
+				PerActorData[ActorIndex].bBlending = true;
+			}
+
+			BlendInstance.TransitionBlendInstance->Setup(From, To, BlendInstance.BlendTime);
+			BlendInstance.bActive = true;
+		}
+	}
+
+	for (FOnDemandInstance& OnDemandInstance : OnDemandInstances)
+	{
+		if (!OnDemandInstance.bActive)
+		{
+			OnDemandInstance.bActive = true;
+			OnDemandInstance.StartTime = WorldTime;
+		}
+	}
+}
+
+float UAnimSharingInstance::CalculateBlendTime(uint8 StateIndex) const
+{
+	checkf(PerStateData.IsValidIndex(StateIndex), TEXT("Invalid State index"));
+	return PerStateData[StateIndex].BlendTime;
+}
+
+void UAnimSharingInstance::RemoveComponent(int32 ComponentIndex)
+{
+	if (PerComponentData.Num() > 1 && ComponentIndex != PerComponentData.Num() - 1)
+	{
+		// Index of the component we will swap with
+		const uint32 SwapIndex = PerComponentData.Num() - 1;
+
+		// Find actor for component we will swap with
+		const uint32 SwapActorIndex = PerComponentData[SwapIndex].ActorIndex;
+
+		// Update component index in the actor to match with ComponentIndex (which it will be swapped with)
+		const uint32 ActorDataComponentIndex = PerActorData[SwapActorIndex].ComponentIndices.IndexOfByKey(SwapIndex);
+		if (ActorDataComponentIndex != INDEX_NONE)
+		{
+			PerActorData[SwapActorIndex].ComponentIndices[ActorDataComponentIndex] = ComponentIndex;
+		}
+	}
+
+	PerComponentData.RemoveAtSwap(ComponentIndex, 1, false);
+}
+
+void UAnimSharingInstance::RemoveBlendInstance(int32 InstanceIndex)
+{
+	FBlendInstance& Instance = BlendInstances[InstanceIndex];
+
+	// Index we could swap with
+	const uint32 SwapIndex = BlendInstances.Num() - 1;
+	if (BlendInstances.Num() > 1 && InstanceIndex != SwapIndex)
+	{
+		FBlendInstance& SwapInstance = BlendInstances[SwapIndex];
+		// Remap all of the actors to point to our new index
+		for (uint32 ActorIndex : SwapInstance.ActorIndices)
+		{
+			PerActorData[ActorIndex].BlendInstanceIndex = InstanceIndex;
+		}
+	}
+
+	BlendInstances.RemoveAtSwap(InstanceIndex, 1, false);
+}
+
+void UAnimSharingInstance::RemoveOnDemandInstance(int32 InstanceIndex)
+{
+	const FOnDemandInstance& Instance = OnDemandInstances[InstanceIndex];
+
+	// Index we could swap with
+	const uint32 SwapIndex = OnDemandInstances.Num() - 1;
+	if (OnDemandInstances.Num() > 1 && InstanceIndex != SwapIndex)
+	{
+		const FOnDemandInstance& SwapInstance = OnDemandInstances[SwapIndex];
+		// Remap all of the actors to point to our new index
+		for (uint32 ActorIndex : SwapInstance.ActorIndices)
+		{
+			// Only remap if it's still part of this instance
+			const bool bPartOfOtherOnDemand = PerActorData[ActorIndex].OnDemandInstanceIndex != InstanceIndex;
+			// Could be swapping with other instance in which case we should update the index
+			const bool bShouldUpdateIndex = !bPartOfOtherOnDemand || (PerActorData[ActorIndex].OnDemandInstanceIndex == SwapIndex);
+			
+			if (bShouldUpdateIndex)
+			{
+				PerActorData[ActorIndex].OnDemandInstanceIndex = InstanceIndex;
+			}		
+		}
+	}
+
+	// Remove and swap 
+	OnDemandInstances.RemoveAtSwap(InstanceIndex, 1, false);
+}
+
+void UAnimSharingInstance::RemoveAdditiveInstance(int32 InstanceIndex)
+{
+	FAdditiveInstance& Instance = AdditiveInstances[InstanceIndex];
+
+	// Index we could swap with
+	const uint32 SwapIndex = AdditiveInstances.Num() - 1;
+	if (AdditiveInstances.Num() > 1 && InstanceIndex != SwapIndex)
+	{
+		FAdditiveInstance& SwapInstance = AdditiveInstances[SwapIndex];
+		// Remap all of the actors to point to our new index
+		if (SwapInstance.ActorIndex != INDEX_NONE)
+		{
+			PerActorData[SwapInstance.ActorIndex].AdditiveInstanceIndex = InstanceIndex;
+		}
+	}
+
+	AdditiveInstances.RemoveAtSwap(InstanceIndex, 1, false);
+}
diff --git a/Engine/Plugins/Developer/AnimationSharing/Source/AnimationSharing/Private/AnimationSharingModule.cpp b/Engine/Plugins/Developer/AnimationSharing/Source/AnimationSharing/Private/AnimationSharingModule.cpp
new file mode 100644
--- /dev/null
+++ b/Engine/Plugins/Developer/AnimationSharing/Source/AnimationSharing/Private/AnimationSharingModule.cpp
@@ -0,0 +1,50 @@
+// Copyright 1998-2018 Epic Games, Inc. All Rights Reserved.
+
+#include "AnimationSharingModule.h"
+#include "Engine/Engine.h"
+#include "Components/SkeletalMeshComponent.h"
+#include "AnimationSharingManager.h"
+
+IMPLEMENT_MODULE( FAnimSharingModule, AnimationSharing);
+
+TMap<const UWorld*, UAnimationSharingManager*> FAnimSharingModule::WorldAnimSharingManagers;
+
+void FAnimSharingModule::StartupModule()
+{
+	FWorldDelegates::OnPostWorldCleanup.AddStatic(&FAnimSharingModule::OnWorldCleanup);
+}
+
+void FAnimSharingModule::AddReferencedObjects(FReferenceCollector& Collector)
+{
+	for (TPair<const UWorld*, UAnimationSharingManager*>& WorldAnimSharingManagerPair : WorldAnimSharingManagers)
+	{
+		Collector.AddReferencedObject(WorldAnimSharingManagerPair.Value, WorldAnimSharingManagerPair.Key);
+	}
+
+#if DEBUG_MATERIALS 
+	for (UMaterialInterface* Material : UAnimationSharingManager::DebugMaterials)
+	{
+		Collector.AddReferencedObject(Material);
+	}
+#endif
+}
+
+bool FAnimSharingModule::CreateAnimationSharingManager(UWorld* InWorld, const UAnimationSharingSetup* Setup)
+{
+	if (InWorld && InWorld->IsGameWorld() && Setup && UAnimationSharingManager::AnimationSharingEnabled() && !WorldAnimSharingManagers.Contains(InWorld))
+	{
+		UAnimationSharingManager* Manager = NewObject<UAnimationSharingManager>(InWorld);
+		Manager->Initialise(Setup);
+		WorldAnimSharingManagers.Add(InWorld, Manager);
+
+		return true;
+	}
+
+	return false;
+}
+
+void FAnimSharingModule::OnWorldCleanup(UWorld* World, bool bSessionEnded, bool bCleanupResources)
+{	
+	WorldAnimSharingManagers.Remove(World);
+}
+
diff --git a/Engine/Plugins/Developer/AnimationSharing/Source/AnimationSharing/Private/AnimationSharingSetup.cpp b/Engine/Plugins/Developer/AnimationSharing/Source/AnimationSharing/Private/AnimationSharingSetup.cpp
new file mode 100644
--- /dev/null
+++ b/Engine/Plugins/Developer/AnimationSharing/Source/AnimationSharing/Private/AnimationSharingSetup.cpp
@@ -0,0 +1,30 @@
+// Copyright 1998-2018 Epic Games, Inc. All Rights Reserved.
+
+#include "AnimationSharingSetup.h"
+
+UAnimationSharingSetup::UAnimationSharingSetup(const FObjectInitializer& ObjectInitializer) : Super(ObjectInitializer)
+{
+	
+}
+
+#if WITH_EDITOR
+void UAnimationSharingSetup::PostLoad()
+{
+	Super::PostLoad();
+
+	/** Ensure all data required for the UI is loaded */
+	for (const FPerSkeletonAnimationSharingSetup& SharingSetup : SkeletonSetups)
+	{ 
+		SharingSetup.Skeleton.LoadSynchronous();
+		SharingSetup.SkeletalMesh.LoadSynchronous();
+		
+		for (const FAnimationStateEntry& Entry : SharingSetup.AnimationStates)
+		{
+			for (const FAnimationSetup& AnimSetup : Entry.AnimationSetups)
+			{
+				AnimSetup.AnimSequence.LoadSynchronous();				
+			}
+		}
+	}
+}
+#endif // WITH_EDITOR
diff --git a/Engine/Plugins/Developer/AnimationSharing/Source/AnimationSharing/Private/TransitionBlendInstance.cpp b/Engine/Plugins/Developer/AnimationSharing/Source/AnimationSharing/Private/TransitionBlendInstance.cpp
new file mode 100644
--- /dev/null
+++ b/Engine/Plugins/Developer/AnimationSharing/Source/AnimationSharing/Private/TransitionBlendInstance.cpp
@@ -0,0 +1,72 @@
+// Copyright 1998-2018 Epic Games, Inc. All Rights Reserved.
+
+#include "TransitionBlendInstance.h"
+#include "AnimationSharingInstances.h"
+#include "AnimationSharingManager.h"
+
+FTransitionBlendInstance::FTransitionBlendInstance() : SkeletalMeshComponent(nullptr), TransitionInstance(nullptr), FromComponent(nullptr), ToComponent(nullptr), BlendTime(0.f), bBlendState(false) {}
+
+void FTransitionBlendInstance::Initialise(USkeletalMeshComponent* InSkeletalMeshComponent, UClass* InAnimationBPClass)
+{
+	if (InSkeletalMeshComponent && InAnimationBPClass)
+	{
+		SkeletalMeshComponent = InSkeletalMeshComponent;
+		SkeletalMeshComponent->SetAnimInstanceClass(InAnimationBPClass);
+		SkeletalMeshComponent->SetComponentTickEnabled(false);	
+		SkeletalMeshComponent->SetForcedLOD(0);
+		SkeletalMeshComponent->VisibilityBasedAnimTickOption = EVisibilityBasedAnimTickOption::AlwaysTickPoseAndRefreshBones;
+		TransitionInstance = Cast<UAnimSharingTransitionInstance>(SkeletalMeshComponent->GetAnimInstance());
+	}
+}
+
+void FTransitionBlendInstance::Setup(USkeletalMeshComponent* InFromComponent, USkeletalMeshComponent* InToComponent, float InBlendTime)
+{
+	UAnimationSharingManager::SetDebugMaterial(SkeletalMeshComponent, 1);
+	SkeletalMeshComponent->SetComponentTickEnabled(true);
+	BlendTime = InBlendTime;
+	if (TransitionInstance)
+	{
+		if (TransitionInstance->bBlendBool)
+		{
+			TransitionInstance->FromComponent = FromComponent = InFromComponent;
+			TransitionInstance->ToComponent = ToComponent = InToComponent;			
+		}
+		else
+		{
+			TransitionInstance->FromComponent = FromComponent = InToComponent;
+			TransitionInstance->ToComponent = ToComponent = InFromComponent;
+		}
+
+		TransitionInstance->BlendTime = InBlendTime;
+		TransitionInstance->bBlendBool = bBlendState = !TransitionInstance->bBlendBool;
+
+		SkeletalMeshComponent->AddTickPrerequisiteComponent(FromComponent);
+		SkeletalMeshComponent->AddTickPrerequisiteComponent(ToComponent);
+	}	
+}
+
+void FTransitionBlendInstance::Stop()
+{
+	if (TransitionInstance)
+	{
+		UAnimationSharingManager::SetDebugMaterial(SkeletalMeshComponent, 0);
+		SkeletalMeshComponent->SetComponentTickEnabled(false);
+		SkeletalMeshComponent->RemoveTickPrerequisiteComponent(FromComponent);
+		SkeletalMeshComponent->RemoveTickPrerequisiteComponent(ToComponent);
+	}
+}
+
+USkeletalMeshComponent* FTransitionBlendInstance::GetComponent() const
+{
+	return SkeletalMeshComponent;
+}
+
+USkeletalMeshComponent* FTransitionBlendInstance::GetToComponent() const
+{
+	return bBlendState == false ? ToComponent : FromComponent;
+}
+
+USkeletalMeshComponent* FTransitionBlendInstance::GetFromComponent() const
+{
+	return bBlendState == false ? FromComponent : ToComponent;
+}
diff --git a/Engine/Plugins/Developer/AnimationSharing/Source/AnimationSharing/Public/AdditiveAnimationInstance.h b/Engine/Plugins/Developer/AnimationSharing/Source/AnimationSharing/Public/AdditiveAnimationInstance.h
new file mode 100644
--- /dev/null
+++ b/Engine/Plugins/Developer/AnimationSharing/Source/AnimationSharing/Public/AdditiveAnimationInstance.h
@@ -0,0 +1,30 @@
+// Copyright 1998-2018 Epic Games, Inc. All Rights Reserved.
+#pragma once
+
+#include "Animation/SkeletalMeshActor.h"
+#include "Components/SkeletalMeshComponent.h"
+
+class UAnimSharingAdditiveInstance;
+class UAnimSequence;
+
+struct ANIMATIONSHARING_API FAdditiveAnimationInstance
+{
+public:
+	FAdditiveAnimationInstance();
+
+	void Initialise(USkeletalMeshComponent* InSkeletalMeshComponent, UClass* InAnimationBP);
+	void Setup(USkeletalMeshComponent* InBaseComponent, UAnimSequence* InAnimSequence);
+	void UpdateBaseComponent(USkeletalMeshComponent* InBaseComponent);
+	void Stop();
+	void Start();
+
+	USkeletalMeshComponent* GetComponent() const;
+	USkeletalMeshComponent* GetBaseComponent() const;	
+
+protected:
+	USkeletalMeshComponent * SkeletalMeshComponent;
+	UAnimSharingAdditiveInstance* AdditiveInstance;
+	UAnimSequence* AdditiveAnimationSequence;
+	USkeletalMeshComponent* BaseComponent;
+	bool bLoopingState;
+};
\ No newline at end of file
diff --git a/Engine/Plugins/Developer/AnimationSharing/Source/AnimationSharing/Public/AnimationSharingInstances.h b/Engine/Plugins/Developer/AnimationSharing/Source/AnimationSharing/Public/AnimationSharingInstances.h
new file mode 100644
--- /dev/null
+++ b/Engine/Plugins/Developer/AnimationSharing/Source/AnimationSharing/Public/AnimationSharingInstances.h
@@ -0,0 +1,80 @@
+// Copyright 1998-2017 Epic Games, Inc. All Rights Reserved.
+
+#pragma once
+
+#include "CoreMinimal.h"
+#include "UObject/ObjectMacros.h"
+#include "Animation/AnimInstance.h"
+#include "AnimationSharingInstances.generated.h"
+
+UCLASS(transient, Blueprintable)
+class ANIMATIONSHARING_API UAnimSharingStateInstance : public UAnimInstance
+{
+	friend class UAnimSharingInstance;
+
+	GENERATED_UCLASS_BODY()
+
+protected:
+	UPROPERTY(EditDefaultsOnly, BlueprintReadOnly, Transient, Category = AnimationSharing)
+	UAnimSequence* AnimationToPlay;
+
+	UPROPERTY(EditDefaultsOnly, BlueprintReadOnly, Transient, Category = AnimationSharing)
+	float PermutationTimeOffset;
+
+	UPROPERTY(EditDefaultsOnly, BlueprintReadOnly, Transient, Category = AnimationSharing)
+	float PlayRate;
+
+	UPROPERTY(EditDefaultsOnly, BlueprintReadOnly, Transient, Category = AnimationSharing)
+	bool bStateBool;
+
+	UFUNCTION(BlueprintCallable, Category = AnimationSharing)
+	void GetInstancedActors(TArray<class AActor*>& Actors);
+
+private:
+	uint8 StateIndex;
+	uint8 ComponentIndex;
+	class UAnimSharingInstance* Instance;
+};
+
+UCLASS(transient, Blueprintable)
+class UAnimSharingTransitionInstance : public UAnimInstance
+{
+	friend struct FTransitionBlendInstance;
+
+	GENERATED_UCLASS_BODY()
+protected:
+	UPROPERTY(EditDefaultsOnly, BlueprintReadOnly, Transient, Category = Transition)
+	TWeakObjectPtr<USkeletalMeshComponent> FromComponent;
+
+	UPROPERTY(EditDefaultsOnly, BlueprintReadOnly, Transient, Category = Transition)
+	TWeakObjectPtr<USkeletalMeshComponent> ToComponent;
+	
+	UPROPERTY(EditDefaultsOnly, BlueprintReadOnly, Transient, Category = Transition)
+	float BlendTime;
+
+	UPROPERTY(EditDefaultsOnly, BlueprintReadOnly, Transient, Category = Transition)
+	bool bBlendBool;
+};
+
+
+UCLASS(transient, Blueprintable)
+class UAnimSharingAdditiveInstance : public UAnimInstance
+{
+	friend struct FAdditiveAnimationInstance;
+
+	GENERATED_UCLASS_BODY()
+protected:
+	UPROPERTY(EditDefaultsOnly, BlueprintReadOnly, Transient, Category = Additive)
+	TWeakObjectPtr<USkeletalMeshComponent> BaseComponent;
+
+	UPROPERTY(EditDefaultsOnly, BlueprintReadOnly, Transient, Category = Additive)
+	TWeakObjectPtr<UAnimSequence> AdditiveAnimation;
+
+	UPROPERTY(EditDefaultsOnly, BlueprintReadOnly, Transient, Category = Additive)
+	float Alpha;
+
+	UPROPERTY(EditDefaultsOnly, BlueprintReadOnly, Transient, Category = Additive)
+	bool bStateBool;
+};
+
+
diff --git a/Engine/Plugins/Developer/AnimationSharing/Source/AnimationSharing/Public/AnimationSharingManager.h b/Engine/Plugins/Developer/AnimationSharing/Source/AnimationSharing/Public/AnimationSharingManager.h
new file mode 100644
--- /dev/null
+++ b/Engine/Plugins/Developer/AnimationSharing/Source/AnimationSharing/Public/AnimationSharingManager.h
@@ -0,0 +1,554 @@
+// Copyright 1998-2018 Epic Games, Inc. All Rights Reserved.
+
+#pragma once
+
+#include "AnimationSharingModule.h"
+#include "AnimationSharingTypes.h"
+
+#include "Tickable.h"
+#include "UObject/SoftObjectPtr.h"
+#include "Engine/SkeletalMesh.h"
+#include "Engine/DeveloperSettings.h"
+#include "Animation/AnimBlueprint.h"
+#include "AdditiveAnimationInstance.h"
+#include "TransitionBlendInstance.h"
+
+#include "AnimationSharingManager.generated.h"
+
+DECLARE_LOG_CATEGORY_EXTERN(LogAnimationSharing, Log, All);
+DECLARE_STATS_GROUP(TEXT("Animation Sharing Manager"), STATGROUP_AnimationSharing, STATCAT_Advanced);
+
+#define DEBUG_MATERIALS (UE_EDITOR && !UE_BUILD_SHIPPING)
+
+DECLARE_DELEGATE_OneParam(FUpdateActorHandle, int32);
+
+class USignificanceManager;
+typedef uint32 AnimationSharingDataHandle;
+
+/** Structure which holds data about a currently in progress blend between two states */
+struct FBlendInstance
+{
+	/** Flag whether or not this instance is currently active */
+	bool bActive;
+	/** Flag whether or not the actor's part of this have been setup as a slave component to the blend actor, this is done so the blend actor atleast ticks once (otherwise it can pop from the previous blend end pose) */
+	bool bBlendStarted;
+	/** Flag whether or not this instance is blending towards an on-demand state */
+	bool bOnDemand;
+	/** World time in seconds at which the blend has finished (calculated at start of blend world time + blend duration) */
+	float EndTime;
+	/** Duration of the blend */
+	float BlendTime;
+	/** State value to blend from */
+	uint8 StateFrom;
+	/** State value to blend to */
+	uint8 StateTo;
+
+	/** Permutation indices from and to which we are blending, used to ensure we 'forward' the actor to the correct MasterPoseComponent when finished blending */
+	uint32 FromPermutationIndex;
+	uint32 ToPermutationIndex;
+
+	/** Actor used for blending between the two states */
+	FTransitionBlendInstance* TransitionBlendInstance;
+	/** Indices of actors who are set up as slaves to BlendActor's main skeletal mesh component */
+	TArray<uint32> ActorIndices;
+
+	/** Optional index into OnDemandInstances from which we are blending */
+	uint32 FromOnDemandInstanceIndex;
+	/** Optional index into OnDemandInstance to which we are blending */
+	uint32 ToOnDemandInstanceIndex;
+};
+
+/** Structure which holds data about a currently running on-demand state animation instance */
+struct FOnDemandInstance
+{
+	/** Flag whether or not instance is active*/
+	bool bActive;
+	bool bBlendActive;
+
+	/** Flag whether or not the component should be 'returned' to the state they were in before the on-demand animation */
+	bool bReturnToPreviousState;
+
+	/** State value which is active */
+	uint8 State;
+
+	/** State value which the components should be set to when the on-demand animation has finished playing (used when !bReturnToPreviousState) */
+	uint8 ForwardState;
+
+	/** Time at which this instance was started */
+	float StartTime;
+
+	/** Time at which this on demand instance should blend out into the 'next' state the actor is in */
+	float StartBlendTime;
+
+	/** World time in seconds at which the animation has finished playing (calculated at start of blend world time + animation sequence length) */
+	float EndTime;
+
+	/** Index into Components array for the current state data which is used for playing the animation*/
+	uint32 UsedPerStateComponentIndex;
+
+	/** Permutation index that we are blending to before the end of the animation */
+	uint32 BlendToPermutationIndex;
+
+	/** Indices of actors who are set up as slaves to the skeletal mesh component running the animation */
+	TArray<uint32> ActorIndices;
+};
+
+struct FAdditiveInstance
+{
+	/** Flag whether or not instance is active */
+	bool bActive;
+
+	/** State index this instance is running */
+	uint8 State;
+
+	/** Time at which this instance finishes */
+	float EndTime;
+
+	/** Current actor indices as part of this instance */
+	uint32 ActorIndex;
+	
+	/** Skeletal mesh component on which the additive animation is applied */
+	USkeletalMeshComponent* BaseComponent;
+
+	/** Actor used for playing the additive animation */
+	FAdditiveAnimationInstance* AdditiveAnimationInstance;
+};
+
+/** Structure which holds data about a unique state which is linked to an enumeration value defined by the user. The data is populated from the user exposed FAnimationStateEntry */
+struct FPerStateData
+{
+	FPerStateData() : bIsOnDemand(false), bIsAdditive(false), BlendTime(0.f), CurrentFrameOnDemandIndex(INDEX_NONE), StateEnumValue(INDEX_NONE), AdditiveAnimationSequence(nullptr) {}
+
+	/** Flag whether or not this state is an on-demand state, this means that we kick off a unique animation when needed */
+	bool bIsOnDemand;
+
+	/** Flag whether or not this state is an additive state */
+	bool bIsAdditive;
+
+	/** Flag whether or not we should return to the previous state, only used when this state is an on-demand one*/
+	bool bReturnToPreviousState;
+
+	/** Flag whether or not ForwardStateValue should be used hwen the animation has finished*/
+	bool bShouldForwardToState;
+
+	/** Duration of blending when blending to this state */
+	float BlendTime;
+
+	/** This is (re-)set every frame, and allows for quickly finding an on-demand instance which was setup this frame */
+	uint32 CurrentFrameOnDemandIndex;
+	/** Number of 'wiggle' frames, this is used when we run out of available entries in Components, if one of the FOnDemandInstance has started NumWiggleFrames ago or earlier,
+	it is used instead of a brand new one */
+	float WiggleTime;
+
+	/** State value to which the actors part of an FOnDemandInstance should be set to when its animation has finished */
+	uint8 ForwardStateValue;
+	/** Enum value linked to this state */
+	uint8 StateEnumValue;
+
+	/** Animation Sequence that is used for Additive States */
+	UAnimSequence* AdditiveAnimationSequence;
+
+	/** Components setup to play animations for this state */
+	TArray<USkeletalMeshComponent*> Components;
+	/** Bits keeping track which of the components are in-use, in case of On Demand state this is managed by FOnDemandInstance, otherwise we clear and populate the flags each frame */
+	TBitArray<> InUseComponentFrameBits;
+	TBitArray<> PreviousInUseComponentFrameBits;
+
+	/** Bits keeping track whether or not any of the slave components requires the master component to tick */
+	TBitArray<> SlaveTickRequiredFrameBits;
+
+	/** Length of the animations used for an on-demand state, array as it could contain different animation permutations */
+	TArray<float> AnimationLengths;
+};
+
+struct FPerComponentData
+{
+	/** Skeletal mesh component registered for this component */
+	USkeletalMeshComponent* Component;
+	/** Index to the owning actor (used to index PerActorData) */
+	int32 ActorIndex;
+};
+
+struct FPerActorData
+{
+	/** Current state value (used to index PerStateData) */
+	uint8 CurrentState;
+	/** Previous state value (used to index PerStateData) */
+	uint8 PreviousState;
+	/** Permutation index (used to index Components array inside of PerStateData) */
+	uint8 PermutationIndex;
+
+	/** Flag whether or not we are currently blending */
+	bool bBlending;
+	/** Flag whether or not we are currently part of an on-demand animation state */
+	bool bRunningOnDemand;
+
+	/** Flag whether or not we are currently part of an additive animation state */
+	bool bRunningAdditive;
+
+	/** Cached significance value */
+	float SignificanceValue;
+
+	/** Flag whether or not this actor requires the master component to tick */
+	bool bRequiresTick;
+
+	/** Index to blend instance which is currently driving this actor's animation */
+	uint32 BlendInstanceIndex;
+	/** Index to on demand instance which is running accord to our current state (or previous state) */
+	uint32 OnDemandInstanceIndex;
+	/** Index to additive instance which is running on top of our state */
+	uint32 AdditiveInstanceIndex;
+
+	/** Indices of the Components owned by this actor (used to index into PerComponentData) */
+	TArray<uint32> ComponentIndices;
+
+	/** Register delegate called when actor is swapped and the handle should be updated */
+	FUpdateActorHandle UpdateActorHandleDelegate;	
+};
+
+template <typename InstanceType>
+struct FInstanceStack
+{
+	~FInstanceStack()
+	{
+		for (InstanceType* Instance : AvailableInstances)
+		{
+			delete Instance;
+		}
+		AvailableInstances.Empty();
+
+		for (InstanceType* Instance : InUseInstances)
+		{
+			delete Instance;
+		}
+		InUseInstances.Empty();
+	}
+
+
+	/** Return whether instance are available */
+	bool InstanceAvailable() const
+	{
+		return AvailableInstances.Num() != 0;
+	}
+
+	/** Get an available instance */
+	InstanceType* GetInstance()
+	{
+		InstanceType* Instance = nullptr;
+		if (AvailableInstances.Num())
+		{
+			Instance = AvailableInstances.Pop(false);
+			InUseInstances.Add(Instance);
+		}
+
+		return Instance;
+	}
+	
+	/** Return instance back */
+	void FreeInstance(InstanceType* Instance)
+	{
+		AvailableInstances.Add(Instance);
+		InUseInstances.RemoveSwap(Instance);
+	}
+	
+	/** Add a new instance to the 'stack' */
+	void AddInstance(InstanceType* Instance)
+	{
+		AvailableInstances.Add(Instance);
+	}
+
+	TArray<InstanceType*> AvailableInstances;
+	TArray<InstanceType*> InUseInstances;
+};
+
+UCLASS()
+class UAnimSharingInstance : public UObject
+{
+	GENERATED_BODY()
+
+public:
+	// Begin UObject overrides
+	virtual void BeginDestroy() override;
+	// End UObject overrides
+
+	/** This uses the StateProcessor to determine the state index the actor is currently in */
+	uint8 DetermineStateForActor(uint32 ActorIndex, bool& bShouldProcess);
+
+	/** Initial set up of all animation sharing data and states */
+	void Setup(UAnimationSharingManager* AnimationSharingManager, const FPerSkeletonAnimationSharingSetup& SkeletonSetup, const FAnimationSharingScalability* ScalabilitySettings, uint32 Index);
+	/** Populates data for a state setup */
+	void SetupState(FPerStateData& StateData, const FAnimationStateEntry& StateEntry, USkeletalMesh* SkeletalMesh, const FPerSkeletonAnimationSharingSetup& SkeletonSetup, uint32 Index);
+
+	/** Retrieves a blend instance, this could either mean reusing an already in progress one or a brand new one (if available according to scalability settings)
+	returns whether or not the setup was successful */
+	uint32 SetupBlend(uint8 FromState, uint8 ToState, uint32 ActorIndex);
+
+	/** Retrieves a blend instance, and sets up a blend from a currently running On-Demand instance to ToState */
+	uint32 SetupBlendFromOnDemand(uint8 ToState, uint32 OnDemandInstanceIndex, uint32 ActorIndex);
+
+	/** Retrieves a blend instance, and sets up a blend between a currently running On-Demand instance and another one which was started this frame */
+	uint32 SetupBlendBetweenOnDemands(uint8 FromOnDemandInstanceIndex, uint32 ToOnDemandInstanceIndex, uint32 ActorIndex);
+
+	/** Retrieves a blend instance, and setups up a blend to an On-Demand instance from a regular animation state */
+	uint32 SetupBlendToOnDemand(uint8 FromState, uint32 ToOnDemandInstanceIndex, uint32 ActorIndex);
+
+	/** Switches between on-demand instances directly, without blending */
+	void SwitchBetweenOnDemands(uint32 FromOnDemandInstanceIndex, uint32 ToOnDemandInstanceIndex, uint32 ActorIndex);
+
+	/** Retrieves a blend instance, this could either mean reusing an already in progress one or a brand new one (if available according to scalability settings)
+	returns an index into OnDemandInstances array or INDEX_NONE if unable to setup an instance */
+	uint32 SetupOnDemandInstance(uint8 StateIndex);
+
+	/** Retrieves an additive instance, these are unique and cannot be reused */
+	uint32 SetupAdditiveInstance(uint8 StateIndex, uint8 FromState, uint8 StateComponentIndex);
+
+	/** Retrieves the blend-time for this specific state */
+	float CalculateBlendTime(uint8 StateIndex) const;
+
+	/** Kicks off the blend and on-demand instances at the end of the current frame tick, this sets up the blend instance with the correct components to blend between */
+	void KickoffInstances();
+
+	/** Ticks all currently running blend instances, checks whether or not the blend is finished and forwards the actor/components to the correct animation state */
+	void TickBlendInstances();
+
+	/** Ticks all currently running on-demand instances, this checks whether or not the animation has finished or if we have to start blending out of the state already */
+	void TickOnDemandInstances();
+
+	template <typename InstanceType>
+	bool DoAnyActorsRequireTicking(const InstanceType& Instance)
+	{
+		bool bShouldTick = false;
+		for (uint32 ActorIndex : Instance.ActorIndices)
+		{
+			if (PerActorData[ActorIndex].bRequiresTick)
+			{
+				bShouldTick = true;
+				break;
+			}
+		}
+
+		return bShouldTick;
+	}
+
+	/** Ticks all currently running additive animation instances, this checks whether or not it has finished yet and sets the base-component as the master component when it has */
+	void TickAdditiveInstances();
+
+	/** Ticks all Actor Data entries and determines their current state, if changed since last tick it will alter their animation accordingly */
+	void TickActorStates();
+
+	/** Ticks various types of debugging data / drawing (not active in shipping build) */
+	void TickDebugInformation();
+
+	/** Ticks all unique animation states, this checks which components are currently used and turns of those which currently don't have any slaves */
+	void TickAnimationStates();
+
+	/** Removal functions which also make sure that indices of other data-structres are correctly remapped */
+	void RemoveBlendInstance(int32 InstanceIndex);
+	void RemoveOnDemandInstance(int32 InstanceIndex);
+	void RemoveAdditiveInstance(int32 InstanceIndex);
+	void RemoveComponent(int32 ComponentIndex);
+
+	/** Removal functions which also make sure the actor is set to the correct master pose component */
+	void RemoveFromCurrentBlend(int32 ActorIndex);
+	void RemoveFromCurrentOnDemand(int32 ActorIndex);
+
+	/** Frees up a Blend instance and resets its state */
+	void FreeBlendInstance(FTransitionBlendInstance* Instance);
+
+	/**  Frees up an Additive Animation instance and resets it state*/
+	void FreeAdditiveInstance(FAdditiveAnimationInstance* Instance);
+
+	/** Sets up all components of an actor to be slaves of Component */
+	void SetMasterComponentForActor(uint32 ActorIndex, USkeletalMeshComponent* Component);
+
+	/** Sets up the correct MasterPoseComponent for the passed in Component and State indices */
+	void SetupSlaveComponent(uint8 CurrentState, uint32 ActorIndex);
+
+	/** Sets up the correct MasterPoseComponent according to the state and permutation indices */
+	void SetPermutationSlaveComponent(uint8 StateIndex, uint32 ActorIndex, uint32 PermutationIndex);
+	
+	/** Determines a permutation index for the given actor and state */
+	uint32 DeterminePermutationIndex(uint32 ActorIndex, uint8 State) const;
+
+	/** Marks the component as either used/not-used, this is used to disable ticking of components which are not in use*/
+	void SetComponentUsage(bool bUsage, uint8 StateIndex, uint32 ComponentIndex);
+
+	/** Sets the whether or not any of the slave components are visible */
+	void SetComponentTick(uint8 StateIndex, uint32 ComponentIndex);
+
+	/** Actors currently registered to be animation driven by the AnimManager using this setup */
+	UPROPERTY(VisibleAnywhere, Transient, Category = AnimationSharing)
+	TArray<AActor*> RegisteredActors;
+	
+	/** Per actor data, matches RegisteredActors*/
+	TArray<FPerActorData> PerActorData;
+	/** Per component state data indexed from FPerActorData.ComponentIndices */
+	TArray<FPerComponentData> PerComponentData;
+
+	/** Array of unique state data */
+	TArray<FPerStateData> PerStateData;
+
+	/** Blend and Additive actors data structure */
+	FInstanceStack<FTransitionBlendInstance> BlendInstanceStack;
+	FInstanceStack<FAdditiveAnimationInstance> AdditiveInstanceStack;
+
+	/** (Blueprint)class instance used for determining the state enum value for each registered actor */
+	UPROPERTY(EditAnywhere, Transient, Category = AnimationSharing)
+	UAnimationSharingStateProcessor* StateProcessor;
+	bool bNativeStateProcessor;
+
+	/** Currently running blend instances */
+	TArray<FBlendInstance> BlendInstances;
+	/** Currently running on-demand instance */
+	TArray<FOnDemandInstance> OnDemandInstances;
+	/** Currently running additive instances */
+	TArray<FAdditiveInstance> AdditiveInstances;
+
+	UPROPERTY(VisibleAnywhere, Transient, Category = AnimationSharing)
+	TArray<UAnimSequence*> UsedAnimationSequences;
+
+	/** Significance manager used for retrieve AI actor significance values */
+	USignificanceManager* SignificanceManager;
+
+	/** Animation sharing manager for the current world */
+	UAnimationSharingManager* AnimSharingManager;
+
+	/** Enum class set up by the user to 'describe' the animation states */
+	UPROPERTY(VisibleAnywhere, Transient, Category = AnimationSharing)
+	UEnum* StateEnum;
+
+	/** Actor to which all the running SkeletalMeshComponents used for the sharing are attached to */
+	UPROPERTY(VisibleAnywhere, Transient, Category = AnimationSharing)
+	AActor* SharingActor;
+
+	/** Platform specific scalability settings */
+	const FAnimationSharingScalability* ScalabilitySettings;
+
+	/** Bounds for the currently used skeletal mesh */
+	FVector SkeletalMeshBounds;
+
+	/** Number of animation setups */
+	uint32 NumSetups;
+
+	/** Holds the current frame world time */
+	float WorldTime;
+};
+
+USTRUCT()
+struct FTickAnimationSharingFunction : public FTickFunction
+{
+	GENERATED_USTRUCT_BODY()
+
+	FTickAnimationSharingFunction() : Manager(nullptr) {}
+
+	// Begin FTickFunction overrides
+	virtual void ExecuteTick(float DeltaTime, enum ELevelTick TickType, ENamedThreads::Type CurrentThread, const FGraphEventRef& MyCompletionGraphEvent) override;
+	virtual FString DiagnosticMessage() override;
+	// End FTickFunction overrides
+	
+	class UAnimationSharingManager* Manager;
+};
+
+
+template<>
+struct TStructOpsTypeTraits<FTickAnimationSharingFunction> : public TStructOpsTypeTraitsBase2<FTickAnimationSharingFunction>
+{
+	enum
+	{
+		WithCopy = false
+	};
+};
+
+UCLASS(config = Engine, defaultconfig)
+class ANIMATIONSHARING_API UAnimationSharingManager : public UObject
+{
+	GENERATED_BODY()
+
+public:
+	// Begin UObject overrides
+	virtual void BeginDestroy() override;
+	virtual UWorld* GetWorld()const override;
+	// End UObject overrides
+
+	/** Returns the AnimationSharing Manager, nullptr if none was set up */
+	UFUNCTION(BlueprintCallable, Category = AnimationSharing, meta = (WorldContext = "WorldContextObject"))
+	static UAnimationSharingManager* GetAnimationSharingManager(UObject* WorldContextObject);
+
+	/** Create an Animation Sharing Manager using the provided Setup */
+	UFUNCTION(BlueprintCallable, Category = AnimationSharing, meta = (WorldContext = "WorldContextObject"))
+	static bool CreateAnimationSharingManager(UObject* WorldContextObject, const UAnimationSharingSetup* Setup);
+
+	/** Register an Actor with this Animation Sharing manager, according to the SharingSkeleton */
+	UFUNCTION(BlueprintCallable, Category = AnimationSharing, meta = (DisplayName = "Register Actor"))
+	void RegisterActorWithSkeletonBP(AActor* InActor, const USkeleton* SharingSkeleton);
+
+	/** Returns whether or not the animation sharing is enabled */
+	UFUNCTION(BlueprintPure, Category = AnimationSharing)
+	static bool AnimationSharingEnabled();
+
+	/** Returns the AnimationSharing Manager for a specific UWorld, nullptr if none was set up */
+	static UAnimationSharingManager* GetManagerForWorld(UWorld* InWorld);
+
+	/** Registers actor with the animation sharing system */
+	void RegisterActor(AActor* InActor, FUpdateActorHandle Delegate);
+
+	/** Registers actor with the animation sharing system according to the SharingSkeleton's sharing setup (if available) */
+	void RegisterActorWithSkeleton(AActor* InActor, const USkeleton* SharingSkeleton, FUpdateActorHandle Delegate);
+
+	/** Unregisters actor with the animation sharing system */
+	void UnregisterActor(AActor* InActor);
+
+	/** Update cached significance for registered actor */
+	void UpdateSignificanceForActorHandle(uint32 InHandle, float InValue);
+
+	/** Ensures all actor data is cleared */
+	void ClearActorData();
+
+	/** Ensures all currently registered actors are removed */
+	void UnregisterAllActors();	
+
+	/** Sets the visibility of currently used Master Pose Components */
+	void SetMasterComponentsVisibility(bool bVisible);
+
+	/** Initialize sharing data structures */
+	void Initialise(const UAnimationSharingSetup* InSetup);
+
+	/** Returns current scalability settings */
+	const FAnimationSharingScalability& GetScalabilitySettings() const;
+
+	void Tick(float DeltaTime);
+	FTickAnimationSharingFunction& GetTickFunction();
+protected:
+	
+	/** Populates all data required for a Skeleton setup */
+	void SetupPerSkeletonData(const FPerSkeletonAnimationSharingSetup& SkeletonSetup);
+
+	/** Dealing with Actor data and handles */
+	uint32 CreateActorHandle(uint8 SkeletonIndex, uint32 ActorIndex) const;
+	uint8 GetSkeletonIndexFromHandle(uint32 InHandle) const;
+	uint32 GetActorIndexFromHandle(uint32 InHandle) const;
+	FPerActorData* GetActorDataByHandle(uint32 InHandle);
+protected:
+	/** Array of unique skeletons, matches PerSkeletonData array entries*/
+	TArray<const USkeleton*> Skeletons;
+
+	/** Sharing data required for the unique Skeleton setups */
+	UPROPERTY(VisibleAnywhere, Transient, Category = AnimationSharing)
+	TArray<UAnimSharingInstance*> PerSkeletonData;
+	
+	/** Platform specific scalability settings */
+	FAnimationSharingScalability ScalabilitySettings;
+
+	/** Tick function for this manager */
+	FTickAnimationSharingFunction TickFunction;
+public:
+#if DEBUG_MATERIALS 	
+	static TArray<UMaterialInterface*> DebugMaterials;	
+#endif
+	static void SetDebugMaterial(USkeletalMeshComponent* Component, uint8 State);
+	static void SetDebugMaterialForActor(UAnimSharingInstance* Data, uint32 ActorIndex, uint8 State);
+
+#if WITH_EDITOR
+	static FName GetPlatformName();
+#endif
+};
diff --git a/Engine/Plugins/Developer/AnimationSharing/Source/AnimationSharing/Public/AnimationSharingModule.h b/Engine/Plugins/Developer/AnimationSharing/Source/AnimationSharing/Public/AnimationSharingModule.h
new file mode 100644
--- /dev/null
+++ b/Engine/Plugins/Developer/AnimationSharing/Source/AnimationSharing/Public/AnimationSharingModule.h
@@ -0,0 +1,36 @@
+// Copyright 1998-2018 Epic Games, Inc. All Rights Reserved.
+
+#pragma once
+
+#include "CoreMinimal.h"
+#include "UObject/ObjectMacros.h"
+#include "UObject/Object.h"
+#include "Modules/ModuleManager.h"
+#include "UObject/GCObject.h"
+#include "Engine/World.h"
+
+class UAnimationSharingManager;
+class UAnimationSharingSetup;
+
+class ANIMATIONSHARING_API FAnimSharingModule : public FDefaultGameModuleImpl, public FGCObject
+{
+public:
+	// Begin IModuleInterface overrides
+	virtual void StartupModule() override;
+	// End IModuleInterface overrides
+
+	// Begin FGCObject overrides
+	virtual void AddReferencedObjects( FReferenceCollector& Collector ) override;
+	// End FGCObject overrides
+
+	FORCEINLINE static UAnimationSharingManager* Get(const UWorld* World)
+	{
+		return WorldAnimSharingManagers.FindRef(World);
+	}
+
+	/** Creates an animation sharing manager for the given UWorld (must be a Game World) */
+	static bool CreateAnimationSharingManager(UWorld* InWorld, const UAnimationSharingSetup* Setup);
+private:	
+	static void OnWorldCleanup(UWorld* World, bool bSessionEnded, bool bCleanupResources);
+	static TMap<const UWorld*, UAnimationSharingManager*> WorldAnimSharingManagers;
+};
diff --git a/Engine/Plugins/Developer/AnimationSharing/Source/AnimationSharing/Public/AnimationSharingSetup.h b/Engine/Plugins/Developer/AnimationSharing/Source/AnimationSharing/Public/AnimationSharingSetup.h
new file mode 100644
--- /dev/null
+++ b/Engine/Plugins/Developer/AnimationSharing/Source/AnimationSharing/Public/AnimationSharingSetup.h
@@ -0,0 +1,27 @@
+// Copyright 1998-2018 Epic Games, Inc. All Rights Reserved.
+
+#pragma once
+
+#include "CoreMinimal.h"
+#include "UObject/ObjectMacros.h"
+#include "UObject/Object.h"
+#include "AnimationSharingTypes.h"
+
+#include "AnimationSharingSetup.generated.h"
+
+UCLASS(hidecategories = Object, Blueprintable, config = Engine)
+class ANIMATIONSHARING_API UAnimationSharingSetup : public UObject
+{
+	GENERATED_UCLASS_BODY()
+public:
+
+#if WITH_EDITOR
+	virtual void PostLoad() override;
+#endif // WITH_EDITOR
+
+	UPROPERTY(EditAnywhere, config, Category = AnimationSharing)
+	TArray<FPerSkeletonAnimationSharingSetup> SkeletonSetups;
+
+	UPROPERTY(EditAnywhere, config, Category = AnimationSharing)
+	FAnimationSharingScalability ScalabilitySettings;
+};
\ No newline at end of file
diff --git a/Engine/Plugins/Developer/AnimationSharing/Source/AnimationSharing/Public/AnimationSharingTypes.h b/Engine/Plugins/Developer/AnimationSharing/Source/AnimationSharing/Public/AnimationSharingTypes.h
new file mode 100644
--- /dev/null
+++ b/Engine/Plugins/Developer/AnimationSharing/Source/AnimationSharing/Public/AnimationSharingTypes.h
@@ -0,0 +1,181 @@
+// Copyright 1998-2018 Epic Games, Inc. All Rights Reserved.
+
+#pragma once
+
+#include "Engine/DeveloperSettings.h"
+#include "PerPlatformProperties.h"
+#include "Animation/AnimBlueprint.h"
+#include "Animation/AnimSequence.h"
+#include "Animation/Skeleton.h"
+#include "Engine/SkeletalMesh.h"
+#include "UObject/Class.h"
+#include "Animation/AnimInstance.h"
+#include "AnimationSharingInstances.h"
+#include "AnimationSharingTypes.generated.h"
+
+USTRUCT()
+struct FAnimationSetup
+{
+	GENERATED_BODY()
+public:
+	FAnimationSetup() : AnimBlueprint(nullptr), AnimSequence(nullptr), NumRandomizedInstances(1), Enabled(true) {}
+	
+	/** Animation Sequence to play for this particular setup */
+	UPROPERTY(EditAnywhere, Category = AnimationSharing)
+	TSoftObjectPtr<UAnimSequence> AnimSequence;
+
+	/** Animation blueprint to use for playing back the Animation Sequence */
+	UPROPERTY(EditAnywhere, AdvancedDisplay, Category = AnimationSharing)
+	TSubclassOf<UAnimSharingStateInstance> AnimBlueprint;	
+
+	/** The number of randomized instances created from this setup, it will create instance with different start time offsets (Length / Number of Instance) * InstanceIndex */
+	UPROPERTY(EditAnywhere, Category = AnimationSharing)
+	FPerPlatformInt NumRandomizedInstances;
+
+	/** Whether or not this setup is enabled for specific platforms */
+	UPROPERTY(EditAnywhere, Category = AnimationSharing)
+	FPerPlatformBool Enabled;
+};
+
+USTRUCT()
+struct FAnimationStateEntry
+{
+	GENERATED_BODY()
+public:
+	FAnimationStateEntry() : State(0), bOnDemand(false), bAdditive(false), BlendTime(0.f), bReturnToPreviousState(false), bSetNextState(false), NextState(0), MaximumNumberOfConcurrentInstances(1), WiggleTimePercentage(0.1f), bRequiresCurves(false) {}
+
+	/** Enum value linked to this state */
+	UPROPERTY(EditAnywhere, Category = AnimationSharing)
+	uint8 State;
+
+	/** Per state animation setup */
+	UPROPERTY(EditAnywhere, Category = AnimationSharing)
+	TArray<FAnimationSetup> AnimationSetups;
+
+	/** Flag whether or not this state is an on-demand state, this means that we kick off a unique animation when needed */
+	UPROPERTY(EditAnywhere, Category = AnimationSharing)
+	bool bOnDemand;
+
+	/** Whether or not this state is an additive state */
+	UPROPERTY(EditAnywhere, Category = AnimationSharing, meta = (EditCondition = "bOnDemand"))
+	bool bAdditive;
+
+	/** Duration of blending when blending to this state */
+	UPROPERTY(EditAnywhere, Category = AnimationSharing, meta = (EditCondition = "!bAdditive"))
+	float BlendTime;
+
+	/** Flag whether or not we should return to the previous state, only used when this state is an on-demand one*/
+	UPROPERTY(EditAnywhere, Category = AnimationSharing, meta = (EditCondition = "bOnDemand"))
+	bool bReturnToPreviousState;
+		
+	UPROPERTY(EditAnywhere, Category = AnimationSharing, meta = (EditCondition = "bOnDemand"))
+	bool bSetNextState;
+
+	/** State value to which the actors part of an on demand state should be set to when its animation has finished */
+	UPROPERTY(EditAnywhere, Category = AnimationSharing, meta = (EditCondition = "bSetNextState"))
+	uint8 NextState;
+
+	/** Number of instances that will be created for this state (platform-specific) */
+	UPROPERTY(EditAnywhere, Category = AnimationSharing, meta=(EditCondition="bOnDemand"))
+	FPerPlatformInt MaximumNumberOfConcurrentInstances;
+
+	/** Percentage of 'wiggle' frames, this is used when we run out of available entries in Components, if one of the on-demand animations has started SequenceLength * WiggleFramePercentage ago or earlier,
+	it is used instead of a brand new one */
+	UPROPERTY(EditAnywhere, Category = AnimationSharing, meta = (EditCondition = "bOnDemand", UIMin="0.0", UIMax="1.0"))
+	float WiggleTimePercentage;
+	
+	/** Whether or not this animation requires curves or morphtargets to function correctly for slave components */
+	UPROPERTY(EditAnywhere, AdvancedDisplay, Category = AnimationSharing)
+	bool bRequiresCurves;
+};
+
+UCLASS(Blueprintable)
+class ANIMATIONSHARING_API UAnimationSharingStateProcessor : public UObject
+{
+	GENERATED_BODY()
+public:
+	UFUNCTION(BlueprintNativeEvent, Category = AnimationSharing)
+	void ProcessActorState(int32& OutState, AActor* InActor, uint8 CurrentState, uint8 OnDemandState, bool& bShouldProcess);
+	void ProcessActorState_Implementation(int32& OutState, AActor* InActor, uint8 CurrentState, uint8 OnDemandState, bool& bShouldProcess)
+	{
+		ProcessActorState_Internal(OutState, InActor, CurrentState, OnDemandState, bShouldProcess);
+	}
+
+	UFUNCTION(BlueprintNativeEvent, Category = AnimationSharing)
+	UEnum* GetAnimationStateEnum();
+	UEnum* GetAnimationStateEnum_Implementation()
+	{
+		return GetAnimationStateEnum_Internal();
+	}
+
+	virtual void ProcessActorStates(TArray<int32>& OutStates, const TArray<AActor*>& InActors, const TArray<uint8>& CurrentStates, const TArray<uint8>& OnDemandStates, TArray<bool>& ShouldProcessFlags)
+	{
+		for (int32 Index = 0; Index < InActors.Num(); ++Index)
+		{
+			ProcessActorState(OutStates[Index], InActors[Index], CurrentStates[Index], OnDemandStates[Index], ShouldProcessFlags[Index]);
+		}
+	}
+
+	UPROPERTY(EditAnywhere, Category = AnimationSharing)
+	TSoftObjectPtr<UEnum> AnimationStateEnum;
+protected:
+	virtual UEnum* GetAnimationStateEnum_Internal() { return AnimationStateEnum.LoadSynchronous(); }
+	virtual void ProcessActorState_Internal(int32& OutState, AActor* InActor, uint8 CurrentState, uint8 OnDemandState, bool& bShouldProcess) {}
+};
+
+USTRUCT()
+struct FPerSkeletonAnimationSharingSetup
+{
+	GENERATED_BODY()
+public:
+	FPerSkeletonAnimationSharingSetup() : Skeleton(nullptr), SkeletalMesh(nullptr), BlendAnimBlueprint(nullptr), AdditiveAnimBlueprint(nullptr), StateProcessorClass(nullptr) {}
+
+	/** Skeleton compatible with the animation sharing setup */
+	UPROPERTY(EditAnywhere, Category = AnimationSharing)
+	TSoftObjectPtr<USkeleton> Skeleton;
+
+	/** Skeletal mesh used to setup skeletal mesh components */
+	UPROPERTY(EditAnywhere, Category = AnimationSharing)
+	TSoftObjectPtr<USkeletalMesh> SkeletalMesh;
+
+	/** Animation blueprint used to perform the blending between states */
+	UPROPERTY(EditAnywhere, Category = AnimationSharing, meta = (DisplayName="Animation Blueprint for Blending"))
+	TSubclassOf<UAnimSharingTransitionInstance> BlendAnimBlueprint;
+	
+	/** Animation blueprint used to apply additive animation on top of other states */
+	UPROPERTY(EditAnywhere, Category = AnimationSharing, meta = (DisplayName = "Animation Blueprint for Additive Animation"))
+	TSubclassOf<UAnimSharingAdditiveInstance> AdditiveAnimBlueprint;
+
+	/** Interface class used when determining which state an actor is in */
+	UPROPERTY(EditAnywhere, Category = AnimationSharing)
+	TSubclassOf<UAnimationSharingStateProcessor> StateProcessorClass;
+
+	/** Definition of different animation states */
+	UPROPERTY(EditAnywhere, Category = AnimationSharing)
+	TArray<FAnimationStateEntry> AnimationStates;
+};
+
+USTRUCT()
+struct FAnimationSharingScalability
+{
+	GENERATED_BODY()
+public:
+	FAnimationSharingScalability() : UseBlendTransitions(true), BlendSignificanceValue(0.5f), MaximumNumberConcurrentBlends(1) {}
+
+	/** Flag whether or not to use blend transitions between states */
+	UPROPERTY(EditAnywhere, Category = AnimationSharing)
+	FPerPlatformBool UseBlendTransitions;
+
+	/** Significance value tied to whether or not a transition should be blended */
+	UPROPERTY(EditAnywhere, Category = AnimationSharing, meta = (EditCondition = "bBlendTransitions"))
+	FPerPlatformFloat BlendSignificanceValue;
+
+	/** Maximum number of blends which can be running concurrently */
+	UPROPERTY(EditAnywhere, Category = AnimationSharing, meta = (EditCondition = "bBlendTransitions"))
+	FPerPlatformInt MaximumNumberConcurrentBlends;
+
+	/** Significance value tied to whether or not the master pose components should be ticking */
+	UPROPERTY(EditAnywhere, Category = AnimationSharing, meta = (EditCondition = "bBlendTransitions"))
+	FPerPlatformFloat TickSignificanceValue;
+};
+
diff --git a/Engine/Plugins/Developer/AnimationSharing/Source/AnimationSharing/Public/TransitionBlendInstance.h b/Engine/Plugins/Developer/AnimationSharing/Source/AnimationSharing/Public/TransitionBlendInstance.h
new file mode 100644
--- /dev/null
+++ b/Engine/Plugins/Developer/AnimationSharing/Source/AnimationSharing/Public/TransitionBlendInstance.h
@@ -0,0 +1,28 @@
+// Copyright 1998-2018 Epic Games, Inc. All Rights Reserved.
+#pragma once
+
+#include "Animation/SkeletalMeshActor.h"
+#include "Components/SkeletalMeshComponent.h"
+
+class UAnimSharingTransitionInstance;
+
+struct ANIMATIONSHARING_API FTransitionBlendInstance
+{
+public:	
+	FTransitionBlendInstance();
+	void Initialise(USkeletalMeshComponent* InSkeletalMeshComponent, UClass* InAnimationBP);
+	void Setup(USkeletalMeshComponent* InFromComponent, USkeletalMeshComponent* InToComponent, float InBlendTime);
+	void Stop();
+
+	USkeletalMeshComponent* GetComponent() const;
+	USkeletalMeshComponent* GetToComponent() const;
+	USkeletalMeshComponent* GetFromComponent() const;
+
+protected:
+	USkeletalMeshComponent * SkeletalMeshComponent;
+	UAnimSharingTransitionInstance* TransitionInstance;
+	USkeletalMeshComponent* FromComponent;
+	USkeletalMeshComponent* ToComponent;
+	float BlendTime;
+	bool bBlendState;
+};
\ No newline at end of file
diff --git a/Engine/Plugins/Developer/AnimationSharing/Source/AnimationSharingEd/AnimationSharingEd.Build.cs b/Engine/Plugins/Developer/AnimationSharing/Source/AnimationSharingEd/AnimationSharingEd.Build.cs
new file mode 100644
--- /dev/null
+++ b/Engine/Plugins/Developer/AnimationSharing/Source/AnimationSharingEd/AnimationSharingEd.Build.cs
@@ -0,0 +1,32 @@
+// Copyright 1998-2018 Epic Games, Inc. All Rights Reserved.
+
+using UnrealBuildTool;
+
+public class AnimationSharingEd : ModuleRules
+{
+	public AnimationSharingEd(ReadOnlyTargetRules Target) : base(Target)
+	{
+        PCHUsage = ModuleRules.PCHUsageMode.UseExplicitOrSharedPCHs;
+		
+        PrivateDependencyModuleNames.AddRange(
+			new string[] {
+				"Core",
+				"CoreUObject",
+				"Engine",
+                "UnrealEd",
+                "AssetTools",
+                "AnimationSharing",
+                "Slate",
+                "SlateCore",
+                "PropertyEditor",
+                "EditorStyle"
+            }
+		);
+
+        PrivateIncludePaths.AddRange(
+            new string[] {
+                "AnimationSharingEd/Private"
+         });       
+    }
+
+}
diff --git a/Engine/Plugins/Developer/AnimationSharing/Source/AnimationSharingEd/Private/AnimationSharingEdModule.cpp b/Engine/Plugins/Developer/AnimationSharing/Source/AnimationSharingEd/Private/AnimationSharingEdModule.cpp
new file mode 100644
--- /dev/null
+++ b/Engine/Plugins/Developer/AnimationSharing/Source/AnimationSharingEd/Private/AnimationSharingEdModule.cpp
@@ -0,0 +1,43 @@
+// Copyright 1998-2018 Epic Games, Inc. All Rights Reserved.
+
+#include "AnimationSharingEdModule.h"
+#include "AssetTypeActions_AnimationSharingSetup.h"
+#include "AssetToolsModule.h"
+#include "SetupDetailsViewCustomizations.h"
+
+IMPLEMENT_MODULE(FAnimSharingEdModule, AnimationSharingEd );
+
+void FAnimSharingEdModule::StartupModule()
+{
+	FAssetToolsModule& AssetToolsModule = FAssetToolsModule::GetModule();
+
+	IAssetTools& AssetTools = AssetToolsModule.Get();
+	AssetAction = new FAssetTypeActions_AnimationSharingSetup();
+	AssetTools.RegisterAssetTypeActions(MakeShareable(AssetAction));
+
+	FPropertyEditorModule& PropertyModule = FModuleManager::LoadModuleChecked<FPropertyEditorModule>("PropertyEditor");
+	{
+		PropertyModule.RegisterCustomPropertyTypeLayout("PerSkeletonAnimationSharingSetup", FOnGetPropertyTypeCustomizationInstance::CreateStatic(&FPerSkeletonAnimationSharingSetupCustomization::MakeInstance));
+		PropertyModule.RegisterCustomPropertyTypeLayout("AnimationStateEntry", FOnGetPropertyTypeCustomizationInstance::CreateStatic(&FAnimationStateEntryCustomization::MakeInstance));		
+
+		PropertyModule.RegisterCustomPropertyTypeLayout("AnimationSetup", FOnGetPropertyTypeCustomizationInstance::CreateStatic(&FAnimationSetupCustomization::MakeInstance));
+	}
+}
+
+void FAnimSharingEdModule::ShutdownModule()
+{
+	if (UObjectInitialized())
+	{
+		FAssetToolsModule& AssetToolsModule = FAssetToolsModule::GetModule();
+		IAssetTools& AssetTools = AssetToolsModule.Get();
+		AssetTools.UnregisterAssetTypeActions(AssetAction->AsShared());
+
+		FPropertyEditorModule& PropertyModule = FModuleManager::LoadModuleChecked<FPropertyEditorModule>("PropertyEditor");
+		{
+			PropertyModule.UnregisterCustomPropertyTypeLayout("PerSkeletonAnimationSharingSetup");
+			PropertyModule.UnregisterCustomPropertyTypeLayout("AnimationStateEntry");
+			PropertyModule.UnregisterCustomPropertyTypeLayout("AnimationSetup");
+		}
+	}
+}
+
diff --git a/Engine/Plugins/Developer/AnimationSharing/Source/AnimationSharingEd/Private/AnimationSharingSetupFactory.cpp b/Engine/Plugins/Developer/AnimationSharing/Source/AnimationSharingEd/Private/AnimationSharingSetupFactory.cpp
new file mode 100644
--- /dev/null
+++ b/Engine/Plugins/Developer/AnimationSharing/Source/AnimationSharingEd/Private/AnimationSharingSetupFactory.cpp
@@ -0,0 +1,22 @@
+// Copyright 1998-2018 Epic Games, Inc. All Rights Reserved.
+
+#include "AnimationSharingSetupFactory.h"
+#include "AnimationSharingSetup.h"
+#include "AssetTypeCategories.h"
+
+UAnimationSharingSetupFactory::UAnimationSharingSetupFactory(const FObjectInitializer& ObjectInitializer)
+	: Super(ObjectInitializer)
+{
+	SupportedClass = UAnimationSharingSetup::StaticClass();
+	bCreateNew = true;
+}
+
+UObject* UAnimationSharingSetupFactory::FactoryCreateNew(UClass* Class, UObject* InParent, FName Name, EObjectFlags Flags, UObject* Context, FFeedbackContext* Warn)
+{
+	return NewObject<UAnimationSharingSetup>(InParent, Class, Name, Flags);
+}
+
+uint32 UAnimationSharingSetupFactory::GetMenuCategories() const
+{
+	return EAssetTypeCategories::Animation;
+}
diff --git a/Engine/Plugins/Developer/AnimationSharing/Source/AnimationSharingEd/Private/SetupDetailsViewCustomizations.cpp b/Engine/Plugins/Developer/AnimationSharing/Source/AnimationSharingEd/Private/SetupDetailsViewCustomizations.cpp
new file mode 100644
--- /dev/null
+++ b/Engine/Plugins/Developer/AnimationSharing/Source/AnimationSharingEd/Private/SetupDetailsViewCustomizations.cpp
@@ -0,0 +1,379 @@
+// Copyright 1998-2018 Epic Games, Inc. All Rights Reserved.
+
+#include "SetupDetailsViewCustomizations.h"
+#include "IDetailChildrenBuilder.h"
+#include "EditorStyleSet.h"
+#include "UObject/UObjectGlobals.h"
+#include "Widgets/Input/STextComboBox.h"
+#include "AnimationSharingTypes.h"
+
+#define LOCTEXT_NAMESPACE "AnimationSharingSetupCustomization"
+
+TSharedRef<IPropertyTypeCustomization> FPerSkeletonAnimationSharingSetupCustomization::MakeInstance()
+{
+	return MakeShareable(new FPerSkeletonAnimationSharingSetupCustomization());
+}
+
+void FPerSkeletonAnimationSharingSetupCustomization::CustomizeHeader(TSharedRef<IPropertyHandle> PropertyHandle, FDetailWidgetRow& HeaderRow, IPropertyTypeCustomizationUtils& CustomizationUtils)
+{
+	SkeletonPropertyHandle = PropertyHandle->GetChildHandle(GET_MEMBER_NAME_CHECKED(FPerSkeletonAnimationSharingSetup, Skeleton));
+	if (SkeletonPropertyHandle.IsValid())
+	{		
+		HeaderRow
+		.NameContent()
+		[
+			SNew(SHorizontalBox)
+			+SHorizontalBox::Slot()
+			.FillWidth(1)
+			.VAlign(VAlign_Center)
+			[
+				// Show the name of the asset or actor
+				SNew(STextBlock)
+				.Font(FEditorStyle::GetFontStyle("PropertyWindow.NormalFont"))
+				.Text(this, &FPerSkeletonAnimationSharingSetupCustomization::GetSkeletonName)
+			]
+		];		
+	}
+}
+
+void FPerSkeletonAnimationSharingSetupCustomization::CustomizeChildren(TSharedRef<IPropertyHandle> StructPropertyHandle, IDetailChildrenBuilder& StructBuilder, IPropertyTypeCustomizationUtils& StructCustomizationUtils)
+{
+	uint32 NumChildren;
+	StructPropertyHandle->GetNumChildren(NumChildren);
+
+	const TArray<FName> SkeletonDisabledProperties = 
+	{
+		GET_MEMBER_NAME_CHECKED(FPerSkeletonAnimationSharingSetup, SkeletalMesh),
+		GET_MEMBER_NAME_CHECKED(FPerSkeletonAnimationSharingSetup, StateProcessorClass),
+		GET_MEMBER_NAME_CHECKED(FPerSkeletonAnimationSharingSetup, BlendAnimBlueprint),
+		GET_MEMBER_NAME_CHECKED(FPerSkeletonAnimationSharingSetup, AdditiveAnimBlueprint)
+	};
+
+	const TArray<FName> EnumDisabledProperties =
+	{
+		GET_MEMBER_NAME_CHECKED(FPerSkeletonAnimationSharingSetup, AnimationStates)
+	};
+
+	TSharedPtr<IPropertyHandle> ProcessorProperty = StructPropertyHandle->GetChildHandle(GET_MEMBER_NAME_CHECKED(FPerSkeletonAnimationSharingSetup, 
+		StateProcessorClass));
+
+	void* StructPtr = nullptr;	
+	StructPropertyHandle->GetValueData(StructPtr);
+
+	for (uint32 ChildIndex = 0; ChildIndex < NumChildren; ++ChildIndex)
+	{
+		TSharedRef<IPropertyHandle> ChildHandle = StructPropertyHandle->GetChildHandle(ChildIndex).ToSharedRef();
+		IDetailPropertyRow& Property = StructBuilder.AddProperty(ChildHandle);
+
+		const FName& PropertyName = ChildHandle->GetProperty()->GetFName();
+
+		/** Properties disabled by an invalid USkeleton */
+		if (SkeletonDisabledProperties.Contains(PropertyName))
+		{
+			Property.IsEnabled(TAttribute<bool>::Create([this]() -> bool
+			{
+				UObject* Object = nullptr;
+				return (SkeletonPropertyHandle->GetValue(Object) == FPropertyAccess::Success) && (Object != nullptr);
+			}));
+		}
+
+		/** Properties disabled by invalid UEnum class */
+		if (EnumDisabledProperties.Contains(PropertyName))
+		{
+			Property.IsEnabled(TAttribute<bool>::Create([ProcessorProperty]() -> bool
+			{
+				const UEnum* EnumClass = GetStateEnumClass(ProcessorProperty);
+				return (EnumClass != nullptr);
+			}));
+		}
+
+		/** Disable additive Anim BP property if there aren't any additive states in the setup */
+		if (StructPtr)
+		{
+			if (PropertyName == GET_MEMBER_NAME_CHECKED(FPerSkeletonAnimationSharingSetup, AdditiveAnimBlueprint))
+			{
+				Property.IsEnabled(TAttribute<bool>::Create([StructPtr]() -> bool
+				{
+					FPerSkeletonAnimationSharingSetup* SetupStruct = (FPerSkeletonAnimationSharingSetup*)StructPtr;
+					if (SetupStruct)
+					{
+						const bool bContainsAdditive = SetupStruct->AnimationStates.ContainsByPredicate([](FAnimationStateEntry& Entry) -> bool
+						{
+							return Entry.bAdditive;
+						});
+
+						return bContainsAdditive;
+					}
+
+					return false;
+				}));
+			}
+		}		
+	}
+}
+
+FText FPerSkeletonAnimationSharingSetupCustomization::GetSkeletonName() const
+{
+	UObject* SkeletonObject;
+	FPropertyAccess::Result Result = SkeletonPropertyHandle->GetValue(SkeletonObject);
+
+	FText Name = LOCTEXT("None", "None");
+	if (Result == FPropertyAccess::Success)
+	{
+		if (SkeletonObject)
+		{
+			Name = FText::AsCultureInvariant(SkeletonObject->GetName());
+		}		
+	}	
+
+	return Name;
+}
+
+TSharedRef<IPropertyTypeCustomization> FAnimationStateEntryCustomization::MakeInstance()
+{
+	return MakeShareable(new FAnimationStateEntryCustomization());
+}
+
+void FAnimationStateEntryCustomization::CustomizeHeader(TSharedRef<IPropertyHandle> PropertyHandle, FDetailWidgetRow& HeaderRow, IPropertyTypeCustomizationUtils& CustomizationUtils)
+{
+	TSharedPtr<IPropertyHandle> ParentHandle = PropertyHandle->GetParentHandle()->GetParentHandle();
+
+	// We make the assumption here that the parent handle is the array part of the FPerSkeletonAnimationSharingSetup	
+	ProcessorPropertyHandle = ParentHandle->GetChildHandle(GET_MEMBER_NAME_CHECKED(FPerSkeletonAnimationSharingSetup, StateProcessorClass));
+	
+	StatePropertyHandle = PropertyHandle->GetChildHandle(GET_MEMBER_NAME_CHECKED(FAnimationStateEntry, State));
+	if (StatePropertyHandle.IsValid())
+	{		
+		HeaderRow
+		.NameContent()
+		[
+			SNew(SHorizontalBox)
+			+SHorizontalBox::Slot()
+			.FillWidth(1)
+			.VAlign(VAlign_Center)
+			[
+				// Show the name of the asset or actor
+				SNew(STextBlock)
+				.Font(FEditorStyle::GetFontStyle("PropertyWindow.NormalFont"))
+				.Text(this, &FAnimationStateEntryCustomization::GetStateName, StatePropertyHandle)
+			]
+		];		
+	}
+}
+
+void FAnimationStateEntryCustomization::CustomizeChildren(TSharedRef<IPropertyHandle> StructPropertyHandle, IDetailChildrenBuilder& StructBuilder, IPropertyTypeCustomizationUtils& StructCustomizationUtils)
+{
+	uint32 NumChildren;
+	StructPropertyHandle->GetNumChildren(NumChildren);
+
+	const TArray<FName> OnDemandProperties = { GET_MEMBER_NAME_CHECKED(FAnimationStateEntry, bReturnToPreviousState), GET_MEMBER_NAME_CHECKED(FAnimationStateEntry, bSetNextState), GET_MEMBER_NAME_CHECKED(FAnimationStateEntry, NextState), GET_MEMBER_NAME_CHECKED(FAnimationStateEntry, WiggleTimePercentage) };
+
+	const TArray<FName> EnumProperties = { GET_MEMBER_NAME_CHECKED(FAnimationStateEntry, State), GET_MEMBER_NAME_CHECKED(FAnimationStateEntry, NextState) };
+
+	TSharedPtr<IPropertyHandle> OnDemandHandle = StructPropertyHandle->GetChildHandle(GET_MEMBER_NAME_CHECKED(FAnimationStateEntry, bOnDemand));
+	TSharedPtr<IPropertyHandle> AdditiveHandle = StructPropertyHandle->GetChildHandle(GET_MEMBER_NAME_CHECKED(FAnimationStateEntry, bAdditive));
+
+	for (uint32 ChildIndex = 0; ChildIndex < NumChildren; ++ChildIndex)
+	{
+		TSharedRef<IPropertyHandle> ChildHandle = StructPropertyHandle->GetChildHandle(ChildIndex).ToSharedRef();
+		
+		// Hide any on-demand settings when either the state is not an on-demand or it is but an additive as well
+		TAttribute<EVisibility> VisibilityAttribute = TAttribute<EVisibility>::Create([OnDemandHandle, AdditiveHandle]() -> EVisibility
+		{
+			if (OnDemandHandle.IsValid() && AdditiveHandle.IsValid())
+			{
+				bool bOnDemandValue = false;
+				OnDemandHandle->GetValue(bOnDemandValue);
+
+				bool bAdditiveValue = false;
+				AdditiveHandle->GetValue(bAdditiveValue);
+
+				return bOnDemandValue && !bAdditiveValue ? EVisibility::Visible : EVisibility::Collapsed;
+			}
+
+			return EVisibility::Visible;
+		});
+				
+		if (EnumProperties.Contains(ChildHandle->GetProperty()->GetFName()))
+		{
+			FDetailWidgetRow& WidgetRow = CreateEnumSelectionWidget(ChildHandle, StructBuilder);
+			if (OnDemandProperties.Contains(ChildHandle->GetProperty()->GetFName()))
+			{
+				WidgetRow.Visibility(VisibilityAttribute);
+			}
+		}
+		else
+		{
+			IDetailPropertyRow& PropertyRow = StructBuilder.AddProperty(ChildHandle);
+			if (OnDemandProperties.Contains(ChildHandle->GetProperty()->GetFName()))
+			{
+				PropertyRow.Visibility(VisibilityAttribute);
+			}
+		}		
+	}
+}
+
+FText FAnimationStateEntryCustomization::GetStateName(TSharedPtr<IPropertyHandle> PropertyHandle) const
+{
+	uint8 EnumValue;
+	FPropertyAccess::Result Result = PropertyHandle->GetValue(EnumValue);
+
+	FText Name = LOCTEXT("None", "None");
+	if (Result == FPropertyAccess::Success)
+	{
+		if (UEnum* EnumClass = GetStateEnumClass(ProcessorPropertyHandle))
+		{
+			Name = EnumClass->GetDisplayNameTextByIndex(EnumValue);
+		}
+		else
+		{
+			FFormatNamedArguments Args;
+			Args.Add(TEXT("EnumIndex"), EnumValue);
+			Name = FText::Format(LOCTEXT("EnumIndexValue", "Enum Index {EnumIndex}"), Args);
+		}
+	}
+
+	return Name;
+}
+
+FDetailWidgetRow& FAnimationStateEntryCustomization::CreateEnumSelectionWidget(TSharedRef<IPropertyHandle> ChildHandle, IDetailChildrenBuilder& StructBuilder)
+{
+	GenerateEnumComboBoxItems();
+
+	TSharedPtr<FString> CurrentlySelected = GetSelectedEnum(ChildHandle);
+
+	FDetailWidgetRow& DetailRow = StructBuilder.AddCustomRow(LOCTEXT("EnumStateSearchLabel", "State"))
+	.NameContent()
+	[
+		ChildHandle->CreatePropertyNameWidget()
+	]
+	.ValueContent()
+	[
+		SNew(SHorizontalBox)		
+		+SHorizontalBox::Slot()
+		[
+			SNew(STextComboBox)
+			.OptionsSource(&ComboBoxItems)
+			.InitiallySelectedItem(CurrentlySelected)
+			.OnSelectionChanged(this, &FAnimationStateEntryCustomization::SelectedEnumChanged, ChildHandle)
+			.Font(FEditorStyle::GetFontStyle("PropertyWindow.NormalFont"))
+		]
+	];
+
+	return DetailRow;
+}
+
+const TArray<TSharedPtr<FString>> FAnimationStateEntryCustomization::GetComboBoxSourceItems() const
+{
+	TArray<TSharedPtr<FString>> Items;
+
+	if (UEnum* EnumClass = GetStateEnumClass(ProcessorPropertyHandle))
+	{
+		const int32 NumEnums = EnumClass->NumEnums();
+		for (int32 Index = 0; Index < NumEnums; ++Index)
+		{
+			Items.Add(MakeShareable(new FString(EnumClass->GetDisplayNameTextByIndex(Index).ToString())));
+		}
+	}
+
+	return Items;
+}
+
+const TSharedPtr<FString> FAnimationStateEntryCustomization::GetSelectedEnum(TSharedPtr<IPropertyHandle> PropertyHandle) const
+{
+	const TSharedPtr<FString>* StringPtr = ComboBoxItems.FindByPredicate([this, PropertyHandle](TSharedPtr<FString> SharedString) {
+		return *SharedString == GetStateName(PropertyHandle).ToString();
+	});
+
+	return StringPtr != nullptr ? *StringPtr : MakeShareable(new FString(GetStateName(PropertyHandle).ToString()));
+}
+
+void FAnimationStateEntryCustomization::SelectedEnumChanged(TSharedPtr<FString> Selection, ESelectInfo::Type SelectInfo, TSharedRef<IPropertyHandle> PropertyHandle)
+{
+	if (UEnum* EnumClass = GetStateEnumClass(ProcessorPropertyHandle))
+	{
+		if (Selection && SelectInfo != ESelectInfo::Type::Direct)
+		{
+			uint8 NewEnumValue = EnumClass->GetValueByIndex(ComboBoxItems.IndexOfByKey(Selection));
+			PropertyHandle->SetValue(NewEnumValue);
+		}
+
+	}
+}
+
+void FAnimationStateEntryCustomization::GenerateEnumComboBoxItems()
+{
+	if (UEnum* EnumClass = GetStateEnumClass(ProcessorPropertyHandle))
+	{
+		if (EnumClass != CachedComboBoxEnumClass)
+		{
+			ComboBoxItems.Empty();
+			const int32 NumEnums = EnumClass->NumEnums();
+			for (int32 Index = 0; Index < NumEnums; ++Index)
+			{
+				ComboBoxItems.Add(MakeShareable(new FString(EnumClass->GetDisplayNameTextByIndex(Index).ToString())));
+			}
+
+			CachedComboBoxEnumClass = EnumClass;
+		}
+	}
+}
+
+UEnum* GetStateEnumClass(const TSharedPtr<IPropertyHandle>& InProperty)
+{
+	UObject* EnumObject = nullptr;
+	if (InProperty.IsValid())
+	{
+		UObject* ProcessorObject = nullptr;
+		InProperty->GetValue(ProcessorObject);
+		UClass* ProcessorClass = Cast<UClass>(ProcessorObject);
+		if (ProcessorClass)
+		{
+			UAnimationSharingStateProcessor* Processor = ProcessorClass->GetDefaultObject<UAnimationSharingStateProcessor>();
+			return Processor->GetAnimationStateEnum();
+		}
+	}
+
+	return nullptr;
+}
+
+
+TSharedRef<IPropertyTypeCustomization> FAnimationSetupCustomization::MakeInstance()
+{
+	return MakeShareable(new FAnimationSetupCustomization());
+}
+
+void FAnimationSetupCustomization::CustomizeChildren(TSharedRef<IPropertyHandle> StructPropertyHandle, IDetailChildrenBuilder& StructBuilder, IPropertyTypeCustomizationUtils& StructCustomizationUtils)
+{
+	uint32 NumChildren;
+	StructPropertyHandle->GetNumChildren(NumChildren);
+	const FName AnimSequencePropertyName = GET_MEMBER_NAME_CHECKED(FAnimationSetup, AnimSequence);
+	AnimSequencePropertyHandle = StructPropertyHandle->GetChildHandle(AnimSequencePropertyName);
+
+	for (uint32 ChildIndex = 0; ChildIndex < NumChildren; ++ChildIndex)
+	{
+		TSharedRef<IPropertyHandle> ChildHandle = StructPropertyHandle->GetChildHandle(ChildIndex).ToSharedRef();
+		IDetailPropertyRow& Property = StructBuilder.AddProperty(ChildHandle);
+		
+		/** Disable all properties if there is not a valid Animation Sequence provided */
+		if (AnimSequencePropertyHandle.IsValid() && (ChildHandle->GetProperty()->GetFName() != AnimSequencePropertyName))
+		{
+			Property.IsEnabled(TAttribute<bool>::Create([this]() -> bool
+			{
+				if ( AnimSequencePropertyHandle.Get())
+				{
+					UObject* ObjectPtr = nullptr;
+					if (AnimSequencePropertyHandle->GetValue(ObjectPtr) == FPropertyAccess::Success)
+					{
+						return ObjectPtr != nullptr;
+					}
+				}
+
+				return false;
+			}));
+		}
+	}
+}
+
+#undef LOCTEXT_NAMESPACE // "AnimationSharingSetupCustomization"
+
diff --git a/Engine/Plugins/Developer/AnimationSharing/Source/AnimationSharingEd/Public/AnimationSharingEdModule.h b/Engine/Plugins/Developer/AnimationSharing/Source/AnimationSharingEd/Public/AnimationSharingEdModule.h
new file mode 100644
--- /dev/null
+++ b/Engine/Plugins/Developer/AnimationSharing/Source/AnimationSharingEd/Public/AnimationSharingEdModule.h
@@ -0,0 +1,22 @@
+// Copyright 1998-2018 Epic Games, Inc. All Rights Reserved.
+
+#pragma once
+
+#include "CoreMinimal.h"
+#include "Modules/ModuleInterface.h"
+#include "Modules/ModuleManager.h"
+
+/**
+* The public interface to this module
+*/
+class FAnimSharingEdModule : public IModuleInterface
+{
+public:
+	FAnimSharingEdModule() {}
+
+	virtual void StartupModule();
+	virtual void ShutdownModule();
+private:
+	class FAssetTypeActions_AnimationSharingSetup* AssetAction;
+};
+
diff --git a/Engine/Plugins/Developer/AnimationSharing/Source/AnimationSharingEd/Public/AnimationSharingSetupFactory.h b/Engine/Plugins/Developer/AnimationSharing/Source/AnimationSharingEd/Public/AnimationSharingSetupFactory.h
new file mode 100644
--- /dev/null
+++ b/Engine/Plugins/Developer/AnimationSharing/Source/AnimationSharingEd/Public/AnimationSharingSetupFactory.h
@@ -0,0 +1,25 @@
+// Copyright 1998-2018 Epic Games, Inc. All Rights Reserved.
+
+#pragma once
+
+#include "CoreMinimal.h"
+#include "UObject/ObjectMacros.h"
+#include "Factories/Factory.h"
+#include "AnimationSharingSetupFactory.generated.h"
+
+struct FAssetData;
+class SWindow;
+
+UCLASS(hidecategories = Object, MinimalAPI)
+class UAnimationSharingSetupFactory : public UFactory
+{
+	GENERATED_UCLASS_BODY()
+
+	//~ Begin UFactory Interface
+	virtual UObject* FactoryCreateNew(UClass* Class, UObject* InParent, FName Name, EObjectFlags Flags, UObject* Context, FFeedbackContext* Warn) override;
+	virtual uint32 GetMenuCategories() const;
+	//~ Begin UFactory Interface
+};
+
+
+
diff --git a/Engine/Plugins/Developer/AnimationSharing/Source/AnimationSharingEd/Public/AssetTypeActions_AnimationSharingSetup.h b/Engine/Plugins/Developer/AnimationSharing/Source/AnimationSharingEd/Public/AssetTypeActions_AnimationSharingSetup.h
new file mode 100644
--- /dev/null
+++ b/Engine/Plugins/Developer/AnimationSharing/Source/AnimationSharingEd/Public/AssetTypeActions_AnimationSharingSetup.h
@@ -0,0 +1,18 @@
+// Copyright 1998-2018 Epic Games, Inc. All Rights Reserved.
+
+#pragma once
+
+#include "CoreMinimal.h"
+#include "AnimationSharingManager.h"
+#include "AssetTypeCategories.h"
+#include "AssetTypeActions_Base.h"
+
+class FAssetTypeActions_AnimationSharingSetup : public FAssetTypeActions_Base
+{
+public:
+	// IAssetTypeActions Implementation
+	virtual FText GetName() const override { return NSLOCTEXT("AssetTypeActions", "AssetTypeActions_AnimationSharingSetup", "AnimationSharingSetup"); }
+	virtual FColor GetTypeColor() const override { return FColor(255, 10, 255); }
+	virtual UClass* GetSupportedClass() const override { return UAnimationSharingManager::StaticClass(); }	
+	virtual uint32 GetCategories() override { return EAssetTypeCategories::Animation; }
+};
diff --git a/Engine/Plugins/Developer/AnimationSharing/Source/AnimationSharingEd/Public/SetupDetailsViewCustomizations.h b/Engine/Plugins/Developer/AnimationSharing/Source/AnimationSharingEd/Public/SetupDetailsViewCustomizations.h
new file mode 100644
--- /dev/null
+++ b/Engine/Plugins/Developer/AnimationSharing/Source/AnimationSharingEd/Public/SetupDetailsViewCustomizations.h
@@ -0,0 +1,74 @@
+// Copyright 1998-2018 Epic Games, Inc. All Rights Reserved.
+
+#pragma once
+
+#include "CoreMinimal.h"
+#include "IDetailCustomization.h"
+#include "PropertyCustomizationHelpers.h"
+
+static UEnum* GetStateEnumClass(const TSharedPtr<IPropertyHandle>& InProperty);
+
+class FPerSkeletonAnimationSharingSetupCustomization : public IPropertyTypeCustomization
+{
+public:
+	static TSharedRef<IPropertyTypeCustomization> MakeInstance();
+
+	FPerSkeletonAnimationSharingSetupCustomization() {}
+	virtual ~FPerSkeletonAnimationSharingSetupCustomization() {}
+
+	/** Begin IPropertyTypeCustomization interface */
+	virtual void CustomizeHeader(TSharedRef<IPropertyHandle> PropertyHandle, FDetailWidgetRow& HeaderRow, IPropertyTypeCustomizationUtils& CustomizationUtils) override;
+	virtual void CustomizeChildren(TSharedRef<IPropertyHandle> StructPropertyHandle, IDetailChildrenBuilder& StructBuilder, IPropertyTypeCustomizationUtils& StructCustomizationUtils) override;
+	/** End IPropertyTypeCustomization interface */
+
+protected:
+	FText GetSkeletonName() const;
+
+	TSharedPtr<IPropertyHandle> SkeletonPropertyHandle;
+};
+
+class FAnimationStateEntryCustomization : public IPropertyTypeCustomization
+{
+public:
+	static TSharedRef<IPropertyTypeCustomization> MakeInstance();
+
+	FAnimationStateEntryCustomization() : StatePropertyHandle(nullptr), ProcessorPropertyHandle(nullptr), CachedComboBoxEnumClass(nullptr) {}
+	virtual ~FAnimationStateEntryCustomization() {}
+
+	/** Begin IPropertyTypeCustomization interface */
+	virtual void CustomizeHeader(TSharedRef<IPropertyHandle> PropertyHandle, FDetailWidgetRow& HeaderRow, IPropertyTypeCustomizationUtils& CustomizationUtils) override;
+	virtual void CustomizeChildren(TSharedRef<IPropertyHandle> StructPropertyHandle, IDetailChildrenBuilder& StructBuilder, IPropertyTypeCustomizationUtils& StructCustomizationUtils) override;
+	/** End IPropertyTypeCustomization interface */
+
+protected:
+	FText GetStateName(TSharedPtr<IPropertyHandle> PropertyHandle) const;
+	FDetailWidgetRow& CreateEnumSelectionWidget(TSharedRef<IPropertyHandle> ChildHandle, IDetailChildrenBuilder& StructBuilder);
+	const TArray<TSharedPtr<FString>> GetComboBoxSourceItems() const;
+	const TSharedPtr<FString> GetSelectedEnum(TSharedPtr<IPropertyHandle> PropertyHandle) const;
+	void SelectedEnumChanged(TSharedPtr<FString> Selection, ESelectInfo::Type SelectInfo, TSharedRef<IPropertyHandle> PropertyHandle);
+	void GenerateEnumComboBoxItems();
+
+protected:
+	TSharedPtr<IPropertyHandle> StatePropertyHandle;
+	TSharedPtr<IPropertyHandle> ProcessorPropertyHandle;
+
+	UEnum* CachedComboBoxEnumClass;
+	TArray<TSharedPtr<FString>> ComboBoxItems;
+};
+
+class FAnimationSetupCustomization : public IPropertyTypeCustomization
+{
+public:
+	static TSharedRef<IPropertyTypeCustomization> MakeInstance();
+
+	FAnimationSetupCustomization() {}
+	virtual ~FAnimationSetupCustomization() {}
+
+	/** Begin IPropertyTypeCustomization interface */
+	virtual void CustomizeHeader(TSharedRef<IPropertyHandle> PropertyHandle, FDetailWidgetRow& HeaderRow, IPropertyTypeCustomizationUtils& CustomizationUtils) override {}
+	virtual void CustomizeChildren(TSharedRef<IPropertyHandle> StructPropertyHandle, IDetailChildrenBuilder& StructBuilder, IPropertyTypeCustomizationUtils& StructCustomizationUtils) override;
+	/** End IPropertyTypeCustomization interface */
+
+protected:
+	TSharedPtr<IPropertyHandle> AnimSequencePropertyHandle;
+};
diff --git a/Engine/Source/Editor/EditorStyle/Private/SlateEditorStyle.cpp b/Engine/Source/Editor/EditorStyle/Private/SlateEditorStyle.cpp
--- a/Engine/Source/Editor/EditorStyle/Private/SlateEditorStyle.cpp
+++ b/Engine/Source/Editor/EditorStyle/Private/SlateEditorStyle.cpp
@@ -229,6 +229,10 @@ void FSlateEditorStyle::FStyle::SetupGeneralStyles()
 		.SetActiveHoveredBrush( IMAGE_BRUSH( "Common/Selection", Icon8x8, SelectionColor ) )
 		.SetInactiveBrush( IMAGE_BRUSH( "Common/Selection", Icon8x8, SelectionColor_Inactive ) )
 		.SetInactiveHoveredBrush( IMAGE_BRUSH( "Common/Selection", Icon8x8, SelectionColor_Inactive ) )
+
+		.SetActiveHighlightedBrush(IMAGE_BRUSH("Common/Selection", Icon8x8, SelectionColor_Inactive))
+		.SetInactiveHighlightedBrush(IMAGE_BRUSH("Common/Selection", Icon8x8, SelectionColor_Inactive))
+
 		.SetTextColor( DefaultForeground )
 		.SetSelectedTextColor(InvertedForeground)
 		.SetDropIndicator_Above(BOX_BRUSH("Common/DropZoneIndicator_Above", FMargin(10.0f / 16.0f, 10.0f / 16.0f, 0, 0), SelectionColor))
@@ -912,12 +916,20 @@ void FSlateEditorStyle::FStyle::SetupGeneralStyles()
 			.SetActiveHoveredBrush(IMAGE_BRUSH("Common/Selection", Icon8x8, SelectionColor))
 			.SetInactiveBrush(IMAGE_BRUSH("Common/Selection", Icon8x8, SelectionColor_Inactive))
 			.SetInactiveHoveredBrush(IMAGE_BRUSH("Common/Selection", Icon8x8, SelectionColor_Inactive))
+
+			.SetActiveHighlightedBrush(IMAGE_BRUSH("Common/Selection", Icon8x8, SelectionColor_Inactive))
+			.SetInactiveHighlightedBrush(IMAGE_BRUSH("Common/Selection", Icon8x8, SelectionColor_Inactive))
 		);
 		Set("TableView.NoHoverTableRow", FTableRowStyle(NormalTableRowStyle)
 			.SetEvenRowBackgroundHoveredBrush(FSlateNoResource())
 			.SetOddRowBackgroundHoveredBrush(FSlateNoResource())
 			.SetActiveHoveredBrush(FSlateNoResource())
 			.SetInactiveHoveredBrush(FSlateNoResource())
+			.SetActiveHoveredBrush(FSlateNoResource())
+			.SetInactiveHoveredBrush(FSlateNoResource())
+
+			.SetActiveHighlightedBrush(FSlateNoResource())
+			.SetInactiveHighlightedBrush(FSlateNoResource())
 			);
 
 
@@ -2203,6 +2215,10 @@ void FSlateEditorStyle::FStyle::SetupGeneralStyles()
 			.SetActiveHoveredBrush( IMAGE_BRUSH( "Common/Selection", Icon8x8, SelectionColor ) )
 			.SetInactiveBrush( IMAGE_BRUSH( "Common/Selection", Icon8x8, SelectionColor ) )
 			.SetInactiveHoveredBrush( IMAGE_BRUSH( "Common/Selection", Icon8x8, SelectionColor )  )
+
+			.SetActiveHighlightedBrush(IMAGE_BRUSH("Common/Selection", Icon8x8, SelectionColor))
+			.SetInactiveHighlightedBrush(IMAGE_BRUSH("Common/Selection", Icon8x8, SelectionColor))
+
 			.SetTextColor( DefaultForeground )
 			.SetSelectedTextColor( InvertedForeground )
 			);
@@ -6148,6 +6164,7 @@ void FSlateEditorStyle::FStyle::SetupClassIconsAndThumbnails()
 			TEXT("AnimComposite"),
 			TEXT("AnimMontage"),
 			TEXT("AnimSequence"),
+			TEXT("AnimationSharingSetup"),
 			TEXT("ApplicationLifecycleComponent"),
 			TEXT("AtmosphericFog"),
 			TEXT("BehaviorTree"),
