Updated copyright notice for Animation Sharing plugin, and fixed initialization order

diff --git a/Engine/Plugins/Developer/AnimationSharing/Source/AnimationSharing/AnimationSharing.Build.cs b/Engine/Plugins/Developer/AnimationSharing/Source/AnimationSharing/AnimationSharing.Build.cs
--- a/Engine/Plugins/Developer/AnimationSharing/Source/AnimationSharing/AnimationSharing.Build.cs
+++ b/Engine/Plugins/Developer/AnimationSharing/Source/AnimationSharing/AnimationSharing.Build.cs
@@ -1,4 +1,4 @@
-// Copyright 1998-2018 Epic Games, Inc. All Rights Reserved.
+// Copyright 1998-2019 Epic Games, Inc. All Rights Reserved.
 
 using UnrealBuildTool;
 
diff --git a/Engine/Plugins/Developer/AnimationSharing/Source/AnimationSharing/Private/AdditiveAnimationInstance.cpp b/Engine/Plugins/Developer/AnimationSharing/Source/AnimationSharing/Private/AdditiveAnimationInstance.cpp
--- a/Engine/Plugins/Developer/AnimationSharing/Source/AnimationSharing/Private/AdditiveAnimationInstance.cpp
+++ b/Engine/Plugins/Developer/AnimationSharing/Source/AnimationSharing/Private/AdditiveAnimationInstance.cpp
@@ -1,4 +1,4 @@
-// Copyright 1998-2018 Epic Games, Inc. All Rights Reserved.
+// Copyright 1998-2019 Epic Games, Inc. All Rights Reserved.
 
 #include "AdditiveAnimationInstance.h"
 #include "AnimationSharingInstances.h"
diff --git a/Engine/Plugins/Developer/AnimationSharing/Source/AnimationSharing/Private/AnimationSharingManager.cpp b/Engine/Plugins/Developer/AnimationSharing/Source/AnimationSharing/Private/AnimationSharingManager.cpp
--- a/Engine/Plugins/Developer/AnimationSharing/Source/AnimationSharing/Private/AnimationSharingManager.cpp
+++ b/Engine/Plugins/Developer/AnimationSharing/Source/AnimationSharing/Private/AnimationSharingManager.cpp
@@ -1,4 +1,4 @@
-// Copyright 1998-2018 Epic Games, Inc. All Rights Reserved.
+// Copyright 1998-2019 Epic Games, Inc. All Rights Reserved.
 
 #include "AnimationSharingManager.h"
 #include "AnimationSharingModule.h"
diff --git a/Engine/Plugins/Developer/AnimationSharing/Source/AnimationSharing/Private/AnimationSharingModule.cpp b/Engine/Plugins/Developer/AnimationSharing/Source/AnimationSharing/Private/AnimationSharingModule.cpp
--- a/Engine/Plugins/Developer/AnimationSharing/Source/AnimationSharing/Private/AnimationSharingModule.cpp
+++ b/Engine/Plugins/Developer/AnimationSharing/Source/AnimationSharing/Private/AnimationSharingModule.cpp
@@ -1,4 +1,4 @@
-// Copyright 1998-2018 Epic Games, Inc. All Rights Reserved.
+// Copyright 1998-2019 Epic Games, Inc. All Rights Reserved.
 
 #include "AnimationSharingModule.h"
 #include "Engine/Engine.h"
diff --git a/Engine/Plugins/Developer/AnimationSharing/Source/AnimationSharing/Private/AnimationSharingSetup.cpp b/Engine/Plugins/Developer/AnimationSharing/Source/AnimationSharing/Private/AnimationSharingSetup.cpp
--- a/Engine/Plugins/Developer/AnimationSharing/Source/AnimationSharing/Private/AnimationSharingSetup.cpp
+++ b/Engine/Plugins/Developer/AnimationSharing/Source/AnimationSharing/Private/AnimationSharingSetup.cpp
@@ -1,4 +1,4 @@
-// Copyright 1998-2018 Epic Games, Inc. All Rights Reserved.
+// Copyright 1998-2019 Epic Games, Inc. All Rights Reserved.
 
 #include "AnimationSharingSetup.h"
 
diff --git a/Engine/Plugins/Developer/AnimationSharing/Source/AnimationSharing/Private/TransitionBlendInstance.cpp b/Engine/Plugins/Developer/AnimationSharing/Source/AnimationSharing/Private/TransitionBlendInstance.cpp
--- a/Engine/Plugins/Developer/AnimationSharing/Source/AnimationSharing/Private/TransitionBlendInstance.cpp
+++ b/Engine/Plugins/Developer/AnimationSharing/Source/AnimationSharing/Private/TransitionBlendInstance.cpp
@@ -1,4 +1,4 @@
-// Copyright 1998-2018 Epic Games, Inc. All Rights Reserved.
+// Copyright 1998-2019 Epic Games, Inc. All Rights Reserved.
 
 #include "TransitionBlendInstance.h"
 #include "AnimationSharingInstances.h"
diff --git a/Engine/Plugins/Developer/AnimationSharing/Source/AnimationSharing/Public/AdditiveAnimationInstance.h b/Engine/Plugins/Developer/AnimationSharing/Source/AnimationSharing/Public/AdditiveAnimationInstance.h
--- a/Engine/Plugins/Developer/AnimationSharing/Source/AnimationSharing/Public/AdditiveAnimationInstance.h
+++ b/Engine/Plugins/Developer/AnimationSharing/Source/AnimationSharing/Public/AdditiveAnimationInstance.h
@@ -1,4 +1,4 @@
-// Copyright 1998-2018 Epic Games, Inc. All Rights Reserved.
+// Copyright 1998-2019 Epic Games, Inc. All Rights Reserved.
 #pragma once
 
 #include "Animation/SkeletalMeshActor.h"
diff --git a/Engine/Plugins/Developer/AnimationSharing/Source/AnimationSharing/Public/AnimationSharingManager.h b/Engine/Plugins/Developer/AnimationSharing/Source/AnimationSharing/Public/AnimationSharingManager.h
--- a/Engine/Plugins/Developer/AnimationSharing/Source/AnimationSharing/Public/AnimationSharingManager.h
+++ b/Engine/Plugins/Developer/AnimationSharing/Source/AnimationSharing/Public/AnimationSharingManager.h
@@ -1,4 +1,4 @@
-// Copyright 1998-2018 Epic Games, Inc. All Rights Reserved.
+// Copyright 1998-2019 Epic Games, Inc. All Rights Reserved.
 
 #pragma once
 
diff --git a/Engine/Plugins/Developer/AnimationSharing/Source/AnimationSharing/Public/AnimationSharingModule.h b/Engine/Plugins/Developer/AnimationSharing/Source/AnimationSharing/Public/AnimationSharingModule.h
--- a/Engine/Plugins/Developer/AnimationSharing/Source/AnimationSharing/Public/AnimationSharingModule.h
+++ b/Engine/Plugins/Developer/AnimationSharing/Source/AnimationSharing/Public/AnimationSharingModule.h
@@ -1,4 +1,4 @@
-// Copyright 1998-2018 Epic Games, Inc. All Rights Reserved.
+// Copyright 1998-2019 Epic Games, Inc. All Rights Reserved.
 
 #pragma once
 
diff --git a/Engine/Plugins/Developer/AnimationSharing/Source/AnimationSharing/Public/AnimationSharingSetup.h b/Engine/Plugins/Developer/AnimationSharing/Source/AnimationSharing/Public/AnimationSharingSetup.h
--- a/Engine/Plugins/Developer/AnimationSharing/Source/AnimationSharing/Public/AnimationSharingSetup.h
+++ b/Engine/Plugins/Developer/AnimationSharing/Source/AnimationSharing/Public/AnimationSharingSetup.h
@@ -1,4 +1,4 @@
-// Copyright 1998-2018 Epic Games, Inc. All Rights Reserved.
+// Copyright 1998-2019 Epic Games, Inc. All Rights Reserved.
 
 #pragma once
 
diff --git a/Engine/Plugins/Developer/AnimationSharing/Source/AnimationSharing/Public/AnimationSharingTypes.h b/Engine/Plugins/Developer/AnimationSharing/Source/AnimationSharing/Public/AnimationSharingTypes.h
--- a/Engine/Plugins/Developer/AnimationSharing/Source/AnimationSharing/Public/AnimationSharingTypes.h
+++ b/Engine/Plugins/Developer/AnimationSharing/Source/AnimationSharing/Public/AnimationSharingTypes.h
@@ -1,4 +1,4 @@
-// Copyright 1998-2018 Epic Games, Inc. All Rights Reserved.
+// Copyright 1998-2019 Epic Games, Inc. All Rights Reserved.
 
 #pragma once
 
@@ -18,7 +18,7 @@ struct FAnimationSetup
 {
 	GENERATED_BODY()
 public:
-	FAnimationSetup() : AnimBlueprint(nullptr), AnimSequence(nullptr), NumRandomizedInstances(1), Enabled(true) {}
+	FAnimationSetup() : AnimSequence(nullptr), AnimBlueprint(nullptr), NumRandomizedInstances(1), Enabled(true) {}
 	
 	/** Animation Sequence to play for this particular setup */
 	UPROPERTY(EditAnywhere, Category = AnimationSharing)
diff --git a/Engine/Plugins/Developer/AnimationSharing/Source/AnimationSharing/Public/TransitionBlendInstance.h b/Engine/Plugins/Developer/AnimationSharing/Source/AnimationSharing/Public/TransitionBlendInstance.h
--- a/Engine/Plugins/Developer/AnimationSharing/Source/AnimationSharing/Public/TransitionBlendInstance.h
+++ b/Engine/Plugins/Developer/AnimationSharing/Source/AnimationSharing/Public/TransitionBlendInstance.h
@@ -1,4 +1,4 @@
-// Copyright 1998-2018 Epic Games, Inc. All Rights Reserved.
+// Copyright 1998-2019 Epic Games, Inc. All Rights Reserved.
 #pragma once
 
 #include "Animation/SkeletalMeshActor.h"
diff --git a/Engine/Plugins/Developer/AnimationSharing/Source/AnimationSharingEd/AnimationSharingEd.Build.cs b/Engine/Plugins/Developer/AnimationSharing/Source/AnimationSharingEd/AnimationSharingEd.Build.cs
--- a/Engine/Plugins/Developer/AnimationSharing/Source/AnimationSharingEd/AnimationSharingEd.Build.cs
+++ b/Engine/Plugins/Developer/AnimationSharing/Source/AnimationSharingEd/AnimationSharingEd.Build.cs
@@ -1,4 +1,4 @@
-// Copyright 1998-2018 Epic Games, Inc. All Rights Reserved.
+// Copyright 1998-2019 Epic Games, Inc. All Rights Reserved.
 
 using UnrealBuildTool;
 
diff --git a/Engine/Plugins/Developer/AnimationSharing/Source/AnimationSharingEd/Private/AnimationSharingEdModule.cpp b/Engine/Plugins/Developer/AnimationSharing/Source/AnimationSharingEd/Private/AnimationSharingEdModule.cpp
--- a/Engine/Plugins/Developer/AnimationSharing/Source/AnimationSharingEd/Private/AnimationSharingEdModule.cpp
+++ b/Engine/Plugins/Developer/AnimationSharing/Source/AnimationSharingEd/Private/AnimationSharingEdModule.cpp
@@ -1,4 +1,4 @@
-// Copyright 1998-2018 Epic Games, Inc. All Rights Reserved.
+// Copyright 1998-2019 Epic Games, Inc. All Rights Reserved.
 
 #include "AnimationSharingEdModule.h"
 #include "AssetTypeActions_AnimationSharingSetup.h"
diff --git a/Engine/Plugins/Developer/AnimationSharing/Source/AnimationSharingEd/Private/AnimationSharingSetupFactory.cpp b/Engine/Plugins/Developer/AnimationSharing/Source/AnimationSharingEd/Private/AnimationSharingSetupFactory.cpp
--- a/Engine/Plugins/Developer/AnimationSharing/Source/AnimationSharingEd/Private/AnimationSharingSetupFactory.cpp
+++ b/Engine/Plugins/Developer/AnimationSharing/Source/AnimationSharingEd/Private/AnimationSharingSetupFactory.cpp
@@ -1,4 +1,4 @@
-// Copyright 1998-2018 Epic Games, Inc. All Rights Reserved.
+// Copyright 1998-2019 Epic Games, Inc. All Rights Reserved.
 
 #include "AnimationSharingSetupFactory.h"
 #include "AnimationSharingSetup.h"
diff --git a/Engine/Plugins/Developer/AnimationSharing/Source/AnimationSharingEd/Private/SetupDetailsViewCustomizations.cpp b/Engine/Plugins/Developer/AnimationSharing/Source/AnimationSharingEd/Private/SetupDetailsViewCustomizations.cpp
--- a/Engine/Plugins/Developer/AnimationSharing/Source/AnimationSharingEd/Private/SetupDetailsViewCustomizations.cpp
+++ b/Engine/Plugins/Developer/AnimationSharing/Source/AnimationSharingEd/Private/SetupDetailsViewCustomizations.cpp
@@ -1,4 +1,4 @@
-// Copyright 1998-2018 Epic Games, Inc. All Rights Reserved.
+// Copyright 1998-2019 Epic Games, Inc. All Rights Reserved.
 
 #include "SetupDetailsViewCustomizations.h"
 #include "IDetailChildrenBuilder.h"
diff --git a/Engine/Plugins/Developer/AnimationSharing/Source/AnimationSharingEd/Public/AnimationSharingEdModule.h b/Engine/Plugins/Developer/AnimationSharing/Source/AnimationSharingEd/Public/AnimationSharingEdModule.h
--- a/Engine/Plugins/Developer/AnimationSharing/Source/AnimationSharingEd/Public/AnimationSharingEdModule.h
+++ b/Engine/Plugins/Developer/AnimationSharing/Source/AnimationSharingEd/Public/AnimationSharingEdModule.h
@@ -1,4 +1,4 @@
-// Copyright 1998-2018 Epic Games, Inc. All Rights Reserved.
+// Copyright 1998-2019 Epic Games, Inc. All Rights Reserved.
 
 #pragma once
 
diff --git a/Engine/Plugins/Developer/AnimationSharing/Source/AnimationSharingEd/Public/AnimationSharingSetupFactory.h b/Engine/Plugins/Developer/AnimationSharing/Source/AnimationSharingEd/Public/AnimationSharingSetupFactory.h
--- a/Engine/Plugins/Developer/AnimationSharing/Source/AnimationSharingEd/Public/AnimationSharingSetupFactory.h
+++ b/Engine/Plugins/Developer/AnimationSharing/Source/AnimationSharingEd/Public/AnimationSharingSetupFactory.h
@@ -1,4 +1,4 @@
-// Copyright 1998-2018 Epic Games, Inc. All Rights Reserved.
+// Copyright 1998-2019 Epic Games, Inc. All Rights Reserved.
 
 #pragma once
 
diff --git a/Engine/Plugins/Developer/AnimationSharing/Source/AnimationSharingEd/Public/AssetTypeActions_AnimationSharingSetup.h b/Engine/Plugins/Developer/AnimationSharing/Source/AnimationSharingEd/Public/AssetTypeActions_AnimationSharingSetup.h
--- a/Engine/Plugins/Developer/AnimationSharing/Source/AnimationSharingEd/Public/AssetTypeActions_AnimationSharingSetup.h
+++ b/Engine/Plugins/Developer/AnimationSharing/Source/AnimationSharingEd/Public/AssetTypeActions_AnimationSharingSetup.h
@@ -1,4 +1,4 @@
-// Copyright 1998-2018 Epic Games, Inc. All Rights Reserved.
+// Copyright 1998-2019 Epic Games, Inc. All Rights Reserved.
 
 #pragma once
 
diff --git a/Engine/Plugins/Developer/AnimationSharing/Source/AnimationSharingEd/Public/SetupDetailsViewCustomizations.h b/Engine/Plugins/Developer/AnimationSharing/Source/AnimationSharingEd/Public/SetupDetailsViewCustomizations.h
--- a/Engine/Plugins/Developer/AnimationSharing/Source/AnimationSharingEd/Public/SetupDetailsViewCustomizations.h
+++ b/Engine/Plugins/Developer/AnimationSharing/Source/AnimationSharingEd/Public/SetupDetailsViewCustomizations.h
@@ -1,4 +1,4 @@
-// Copyright 1998-2018 Epic Games, Inc. All Rights Reserved.
+// Copyright 1998-2019 Epic Games, Inc. All Rights Reserved.
 
 #pragma once
 
