Allow UAnimationSharingStateProcessor-based Blueprint classes to be nativized at cook time.

Change summary:
- BlueprintNativeEvent implementation methods declared 'virtual' so that the C++ compiler allows nativized BPs to override them in a derived Blueprint's generated C++ code.

#rnx

diff --git a/Engine/Plugins/Developer/AnimationSharing/Source/AnimationSharing/Public/AnimationSharingTypes.h b/Engine/Plugins/Developer/AnimationSharing/Source/AnimationSharing/Public/AnimationSharingTypes.h
--- a/Engine/Plugins/Developer/AnimationSharing/Source/AnimationSharing/Public/AnimationSharingTypes.h
+++ b/Engine/Plugins/Developer/AnimationSharing/Source/AnimationSharing/Public/AnimationSharingTypes.h
@@ -96,14 +96,14 @@ class ANIMATIONSHARING_API UAnimationSharingStateProcessor : public UObject
 public:
 	UFUNCTION(BlueprintNativeEvent, Category = AnimationSharing)
 	void ProcessActorState(int32& OutState, AActor* InActor, uint8 CurrentState, uint8 OnDemandState, bool& bShouldProcess);
-	void ProcessActorState_Implementation(int32& OutState, AActor* InActor, uint8 CurrentState, uint8 OnDemandState, bool& bShouldProcess)
+	virtual void ProcessActorState_Implementation(int32& OutState, AActor* InActor, uint8 CurrentState, uint8 OnDemandState, bool& bShouldProcess)
 	{
 		ProcessActorState_Internal(OutState, InActor, CurrentState, OnDemandState, bShouldProcess);
 	}
 
 	UFUNCTION(BlueprintNativeEvent, Category = AnimationSharing)
 	UEnum* GetAnimationStateEnum();
-	UEnum* GetAnimationStateEnum_Implementation()
+	virtual UEnum* GetAnimationStateEnum_Implementation()
 	{
 		return GetAnimationStateEnum_Internal();
 	}
