Various animation sharing plugin fixes

diff --git a/Engine/Plugins/Developer/AnimationSharing/Source/AnimationSharing/Private/AdditiveAnimationInstance.cpp b/Engine/Plugins/Developer/AnimationSharing/Source/AnimationSharing/Private/AdditiveAnimationInstance.cpp
--- a/Engine/Plugins/Developer/AnimationSharing/Source/AnimationSharing/Private/AdditiveAnimationInstance.cpp
+++ b/Engine/Plugins/Developer/AnimationSharing/Source/AnimationSharing/Private/AdditiveAnimationInstance.cpp
@@ -10,14 +10,18 @@ FAdditiveAnimationInstance::FAdditiveAnimationInstance() : SkeletalMeshComponent
 
 void FAdditiveAnimationInstance::Initialise(USkeletalMeshComponent* InSkeletalMeshComponent, UClass* InAnimationBPClass)
 {
-	if (InSkeletalMeshComponent && InAnimationBPClass)
+	if (InSkeletalMeshComponent)
 	{
 		SkeletalMeshComponent = InSkeletalMeshComponent;
-		SkeletalMeshComponent->SetAnimInstanceClass(InAnimationBPClass);
+		if (InAnimationBPClass)
+		{
+			SkeletalMeshComponent->SetAnimInstanceClass(InAnimationBPClass);
+			AdditiveInstance = Cast<UAnimSharingAdditiveInstance>(SkeletalMeshComponent->GetAnimInstance());
+		}
+		
 		SkeletalMeshComponent->SetComponentTickEnabled(false);	
 		SkeletalMeshComponent->SetForcedLOD(1);
 		SkeletalMeshComponent->VisibilityBasedAnimTickOption = EVisibilityBasedAnimTickOption::AlwaysTickPoseAndRefreshBones;
-		AdditiveInstance = Cast<UAnimSharingAdditiveInstance>(SkeletalMeshComponent->GetAnimInstance());
 	}
 }
 
diff --git a/Engine/Plugins/Developer/AnimationSharing/Source/AnimationSharing/Private/AnimationSharingManager.cpp b/Engine/Plugins/Developer/AnimationSharing/Source/AnimationSharing/Private/AnimationSharingManager.cpp
--- a/Engine/Plugins/Developer/AnimationSharing/Source/AnimationSharing/Private/AnimationSharingManager.cpp
+++ b/Engine/Plugins/Developer/AnimationSharing/Source/AnimationSharing/Private/AnimationSharingManager.cpp
@@ -221,14 +221,22 @@ const FAnimationSharingScalability& UAnimationSharingManager::GetScalabilitySett
 void UAnimationSharingManager::SetupPerSkeletonData(const FPerSkeletonAnimationSharingSetup& SkeletonSetup)
 {
 	const USkeleton* Skeleton = SkeletonSetup.Skeleton.LoadSynchronous();
-	UEnum* StateEnum = SkeletonSetup.StateProcessorClass->GetDefaultObject<UAnimationSharingStateProcessor>()->GetAnimationStateEnum();
-	if (Skeleton && StateEnum)
+	UAnimationSharingStateProcessor* Processor = SkeletonSetup.StateProcessorClass ?SkeletonSetup.StateProcessorClass->GetDefaultObject<UAnimationSharingStateProcessor>() : nullptr;
+	UEnum* StateEnum = Processor ? Processor->GetAnimationStateEnum() : nullptr;
+	if (Skeleton && StateEnum && Processor)
 	{
 		UAnimSharingInstance* Data = NewObject<UAnimSharingInstance>(this);
 		PerSkeletonData.Add(Data);
 		Skeletons.Add(Skeleton);
 		Data->Setup(this, SkeletonSetup, &ScalabilitySettings, Skeletons.Num() - 1);
 	}
+	else
+	{
+		UE_LOG(LogAnimationSharing, Error, TEXT("Invalid Skeleton (%s), State Enum (%s) or State Processor (%s)!"), 
+			Skeleton ? *Skeleton->GetName() : TEXT("None"),
+			StateEnum ? *StateEnum->GetName() : TEXT("None"),
+			Processor ? *Processor->GetName() : TEXT("None"));
+	}
 }
 
 uint32 UAnimationSharingManager::CreateActorHandle(uint8 SkeletonIndex, uint32 ActorIndex) const
@@ -386,7 +394,7 @@ void UAnimationSharingManager::RegisterActorWithSkeleton(AActor* InActor, const
 		{
 			return (Skeleton == SharingSkeleton) || (Skeleton->IsCompatible(SharingSkeleton));
 		});
-		ensureMsgf(ArrayIndex != INDEX_NONE, TEXT("Invalid skeleton for which there is no sharing setup available!"));
+		UE_LOG(LogAnimationSharing, Error, TEXT("Invalid skeleton (%s) for which there is no sharing setup available!"), SharingSkeleton ? *SharingSkeleton->GetName() : TEXT("None"));
 		return ArrayIndex;
 	}();
 
@@ -397,43 +405,45 @@ void UAnimationSharingManager::RegisterActorWithSkeleton(AActor* InActor, const
 		checkf(OwnedComponents.Num(), TEXT("No SkeletalMeshComponents found in actor!"));
 
 		UAnimSharingInstance* Data = PerSkeletonData[Handle];
+		if (Data->AnimSharingManager != nullptr)
+		{
+			// Register the actor
+			const int32 ActorIndex = Data->RegisteredActors.Add(InActor);
 
-		// Register the actor
-		const int32 ActorIndex = Data->RegisteredActors.Add(InActor);
+			FPerActorData& ActorData = Data->PerActorData.AddZeroed_GetRef();
+			ActorData.BlendInstanceIndex = ActorData.OnDemandInstanceIndex = ActorData.AdditiveInstanceIndex = INDEX_NONE;
+			ActorData.SignificanceValue = Data->SignificanceManager->GetSignificance(InActor);
+			ActorData.UpdateActorHandleDelegate = CallbackDelegate;
 
-		FPerActorData& ActorData = Data->PerActorData.AddZeroed_GetRef();
-		ActorData.BlendInstanceIndex = ActorData.OnDemandInstanceIndex = ActorData.AdditiveInstanceIndex = INDEX_NONE;
-		ActorData.SignificanceValue = Data->SignificanceManager->GetSignificance(InActor);
-		ActorData.UpdateActorHandleDelegate = CallbackDelegate;
+			bool bShouldProcess = true;
+			ActorData.CurrentState = ActorData.PreviousState = Data->DetermineStateForActor(ActorIndex, bShouldProcess);
 
-		bool bShouldProcess = true;
-		ActorData.CurrentState = ActorData.PreviousState = Data->DetermineStateForActor(ActorIndex, bShouldProcess);
+			for (USkeletalMeshComponent* Component : OwnedComponents)
+			{
+				FPerComponentData& ComponentData = Data->PerComponentData.AddZeroed_GetRef();
+				ComponentData.ActorIndex = ActorIndex;
+				ComponentData.Component = Component;
 
-		for (USkeletalMeshComponent* Component : OwnedComponents)
-		{
-			FPerComponentData& ComponentData = Data->PerComponentData.AddZeroed_GetRef();
-			ComponentData.ActorIndex = ActorIndex;
-			ComponentData.Component = Component;
+				Component->PrimaryComponentTick.bCanEverTick = false;
+				Component->SetComponentTickEnabled(false);
+				Component->bIgnoreMasterPoseComponentLOD = true;
 
-			Component->PrimaryComponentTick.bCanEverTick = false;
-			Component->SetComponentTickEnabled(false);
-			Component->bIgnoreMasterPoseComponentLOD = true;
+				ActorData.ComponentIndices.Add(Data->PerComponentData.Num() - 1);
 
-			ActorData.ComponentIndices.Add(Data->PerComponentData.Num() - 1);
+				const int32 ComponentIndex = Data->PerComponentData.Num() - 1;
+				Data->SetupSlaveComponent(ActorData.CurrentState, ActorIndex);
+			}
 
-			const int32 ComponentIndex = Data->PerComponentData.Num() - 1;
-			Data->SetupSlaveComponent(ActorData.CurrentState, ActorIndex);
-		}
+			if (Data->PerStateData[ActorData.CurrentState].bIsOnDemand && ActorData.OnDemandInstanceIndex != INDEX_NONE)
+			{
+				// We will have setup an on-demand instance so we need to kick it off here before we next tick
+				Data->OnDemandInstances[ActorData.OnDemandInstanceIndex].bActive = true;
+				Data->OnDemandInstances[ActorData.OnDemandInstanceIndex].StartTime = Data->WorldTime;
+			}
 
-		if (Data->PerStateData[ActorData.CurrentState].bIsOnDemand && ActorData.OnDemandInstanceIndex != INDEX_NONE)
-		{
-			// We will have setup an on-demand instance so we need to kick it off here before we next tick
-			Data->OnDemandInstances[ActorData.OnDemandInstanceIndex].bActive = true;
-			Data->OnDemandInstances[ActorData.OnDemandInstanceIndex].StartTime = Data->WorldTime;
+			const int32 ActorHandle = CreateActorHandle(Handle, ActorIndex);
+			ActorData.UpdateActorHandleDelegate.ExecuteIfBound(ActorHandle);
 		}
-
-		const int32 ActorHandle = CreateActorHandle(Handle, ActorIndex);
-		ActorData.UpdateActorHandleDelegate.ExecuteIfBound(ActorHandle);
 	}
 }
 
@@ -603,22 +613,34 @@ void UAnimationSharingManager::SetMasterComponentsVisibility(bool bVisible)
 
 		for (FTransitionBlendInstance* Instance : Data->BlendInstanceStack.AvailableInstances)
 		{
-			Instance->GetComponent()->SetVisibility(bVisible);
+			if (USceneComponent* Component = Instance->GetComponent())
+			{
+				Component->SetVisibility(bVisible);
+			}
 		}
 
 		for (FTransitionBlendInstance* Instance : Data->BlendInstanceStack.InUseInstances)
 		{
-			Instance->GetComponent()->SetVisibility(bVisible);
+			if (USceneComponent* Component = Instance->GetComponent())
+			{
+				Component->SetVisibility(bVisible);
+			}
 		}
 
 		for (FAdditiveAnimationInstance* Instance : Data->AdditiveInstanceStack.AvailableInstances)
 		{
-			Instance->GetComponent()->SetVisibility(bVisible);
+			if (USceneComponent* Component = Instance->GetComponent())
+			{
+				Component->SetVisibility(bVisible);
+			}
 		}
 
 		for (FAdditiveAnimationInstance* Instance : Data->AdditiveInstanceStack.InUseInstances)
 		{
-			Instance->GetComponent()->SetVisibility(bVisible);
+			if (USceneComponent* Component = Instance->GetComponent())
+			{
+				Component->SetVisibility(bVisible);
+			}
 		}
 	}
 }
diff --git a/Engine/Plugins/Developer/AnimationSharing/Source/AnimationSharing/Private/TransitionBlendInstance.cpp b/Engine/Plugins/Developer/AnimationSharing/Source/AnimationSharing/Private/TransitionBlendInstance.cpp
--- a/Engine/Plugins/Developer/AnimationSharing/Source/AnimationSharing/Private/TransitionBlendInstance.cpp
+++ b/Engine/Plugins/Developer/AnimationSharing/Source/AnimationSharing/Private/TransitionBlendInstance.cpp
@@ -8,14 +8,18 @@ FTransitionBlendInstance::FTransitionBlendInstance() : SkeletalMeshComponent(nul
 
 void FTransitionBlendInstance::Initialise(USkeletalMeshComponent* InSkeletalMeshComponent, UClass* InAnimationBPClass)
 {
-	if (InSkeletalMeshComponent && InAnimationBPClass)
+	if (InSkeletalMeshComponent)
 	{
 		SkeletalMeshComponent = InSkeletalMeshComponent;
-		SkeletalMeshComponent->SetAnimInstanceClass(InAnimationBPClass);
+		if (InAnimationBPClass)
+		{
+			SkeletalMeshComponent->SetAnimInstanceClass(InAnimationBPClass);
+			TransitionInstance = Cast<UAnimSharingTransitionInstance>(SkeletalMeshComponent->GetAnimInstance());
+		}
+		
 		SkeletalMeshComponent->SetComponentTickEnabled(false);	
 		SkeletalMeshComponent->SetForcedLOD(0);
-		SkeletalMeshComponent->VisibilityBasedAnimTickOption = EVisibilityBasedAnimTickOption::AlwaysTickPoseAndRefreshBones;
-		TransitionInstance = Cast<UAnimSharingTransitionInstance>(SkeletalMeshComponent->GetAnimInstance());
+		SkeletalMeshComponent->VisibilityBasedAnimTickOption = EVisibilityBasedAnimTickOption::AlwaysTickPoseAndRefreshBones;		
 	}
 }
 
diff --git a/Engine/Plugins/Developer/AnimationSharing/Source/AnimationSharingEd/Private/SetupDetailsViewCustomizations.cpp b/Engine/Plugins/Developer/AnimationSharing/Source/AnimationSharingEd/Private/SetupDetailsViewCustomizations.cpp
--- a/Engine/Plugins/Developer/AnimationSharing/Source/AnimationSharingEd/Private/SetupDetailsViewCustomizations.cpp
+++ b/Engine/Plugins/Developer/AnimationSharing/Source/AnimationSharingEd/Private/SetupDetailsViewCustomizations.cpp
@@ -6,6 +6,7 @@
 #include "UObject/UObjectGlobals.h"
 #include "Widgets/Input/STextComboBox.h"
 #include "AnimationSharingTypes.h"
+#include "Templates/SharedPointer.h"
 
 #define LOCTEXT_NAMESPACE "AnimationSharingSetupCustomization"
 
@@ -337,12 +338,48 @@ UEnum* GetStateEnumClass(const TSharedPtr<IPropertyHandle>& InProperty)
 	return nullptr;
 }
 
-
 TSharedRef<IPropertyTypeCustomization> FAnimationSetupCustomization::MakeInstance()
 {
 	return MakeShareable(new FAnimationSetupCustomization());
 }
 
+void FAnimationSetupCustomization::CustomizeHeader(TSharedRef<IPropertyHandle> PropertyHandle, FDetailWidgetRow& HeaderRow, IPropertyTypeCustomizationUtils& CustomizationUtils)
+{
+	const FName AnimSequencePropertyName = GET_MEMBER_NAME_CHECKED(FAnimationSetup, AnimSequence);
+	AnimSequencePropertyHandle = PropertyHandle->GetChildHandle(AnimSequencePropertyName);
+	
+	HeaderRow
+	.NameContent()
+	[
+		SNew(SHorizontalBox)
+		+SHorizontalBox::Slot()
+		.FillWidth(1)
+		.VAlign(VAlign_Center)
+		[
+			// Show the name of the asset or actor
+			SNew(STextBlock)
+			.Font(FEditorStyle::GetFontStyle("PropertyWindow.NormalFont"))
+			.Text_Lambda([this]() -> FText
+			{
+				if (AnimSequencePropertyHandle.IsValid())
+				{
+					FText PropertyValueAsName;
+					if (AnimSequencePropertyHandle->GetValueAsFormattedText(PropertyValueAsName) == FPropertyAccess::Success)
+					{
+						return PropertyValueAsName;
+					}
+				}
+
+				return LOCTEXT("None", "None");
+			})			
+		]
+	];
+	/*.ValueContent()
+	[
+		PropertyHandle->CreatePropertyValueWidget()
+	];*/	
+}
+
 void FAnimationSetupCustomization::CustomizeChildren(TSharedRef<IPropertyHandle> StructPropertyHandle, IDetailChildrenBuilder& StructBuilder, IPropertyTypeCustomizationUtils& StructCustomizationUtils)
 {
 	uint32 NumChildren;
diff --git a/Engine/Plugins/Developer/AnimationSharing/Source/AnimationSharingEd/Public/SetupDetailsViewCustomizations.h b/Engine/Plugins/Developer/AnimationSharing/Source/AnimationSharingEd/Public/SetupDetailsViewCustomizations.h
--- a/Engine/Plugins/Developer/AnimationSharing/Source/AnimationSharingEd/Public/SetupDetailsViewCustomizations.h
+++ b/Engine/Plugins/Developer/AnimationSharing/Source/AnimationSharingEd/Public/SetupDetailsViewCustomizations.h
@@ -65,10 +65,13 @@ public:
 	virtual ~FAnimationSetupCustomization() {}
 
 	/** Begin IPropertyTypeCustomization interface */
-	virtual void CustomizeHeader(TSharedRef<IPropertyHandle> PropertyHandle, FDetailWidgetRow& HeaderRow, IPropertyTypeCustomizationUtils& CustomizationUtils) override {}
+	virtual void CustomizeHeader(TSharedRef<IPropertyHandle> PropertyHandle, FDetailWidgetRow& HeaderRow, IPropertyTypeCustomizationUtils& CustomizationUtils) override;
 	virtual void CustomizeChildren(TSharedRef<IPropertyHandle> StructPropertyHandle, IDetailChildrenBuilder& StructBuilder, IPropertyTypeCustomizationUtils& StructCustomizationUtils) override;
 	/** End IPropertyTypeCustomization interface */
 
+protected:
+	FText GetAnimationName(TSharedPtr<IPropertyHandle> PropertyHandle) const;
+
 protected:
 	TSharedPtr<IPropertyHandle> AnimSequencePropertyHandle;
 };
