//UE4/Dev-Anim - Run Automated Tests - Invalid skeleton (UE4_Mannequin_Skeleton) for which there is no sharing setup available!
#fix Move error message into correct clause
#misc Shuffled content a bit to ensure it always runs in the correct sequence of events, and fixed up default-reference lambda variable captures

diff --git a/Engine/Plugins/Developer/AnimationSharing/Source/AnimationSharing/Private/AnimationSharingManager.cpp b/Engine/Plugins/Developer/AnimationSharing/Source/AnimationSharing/Private/AnimationSharingManager.cpp
--- a/Engine/Plugins/Developer/AnimationSharing/Source/AnimationSharing/Private/AnimationSharingManager.cpp
+++ b/Engine/Plugins/Developer/AnimationSharing/Source/AnimationSharing/Private/AnimationSharingManager.cpp
@@ -353,7 +353,7 @@ void UAnimationSharingManager::RegisterActor(AActor* InActor, FUpdateActorHandle
 		InActor->GetComponents<USkeletalMeshComponent>(OwnedComponents);
 		checkf(OwnedComponents.Num(), TEXT("No SkeletalMeshComponents found in actor!"));
 
-		const USkeleton* UsedSkeleton = [&]()
+		const USkeleton* UsedSkeleton = [&OwnedComponents]()
 		{
 			USkeleton* CurrentSkeleton = nullptr;
 			for (USkeletalMeshComponent* SkeletalMeshComponent : OwnedComponents)
@@ -388,13 +388,12 @@ void UAnimationSharingManager::RegisterActorWithSkeleton(AActor* InActor, const
 		return;
 	}
 
-	const AnimationSharingDataHandle Handle = [&]() -> uint32
+	const AnimationSharingDataHandle Handle = [this, SharingSkeleton]() -> uint32
 	{
-		uint32 ArrayIndex = Skeletons.IndexOfByPredicate([&](const USkeleton* Skeleton)
+		uint32 ArrayIndex = Skeletons.IndexOfByPredicate([SharingSkeleton](const USkeleton* Skeleton)
 		{
 			return (Skeleton == SharingSkeleton) || (Skeleton->IsCompatible(SharingSkeleton));
 		});
-		UE_LOG(LogAnimationSharing, Error, TEXT("Invalid skeleton (%s) for which there is no sharing setup available!"), SharingSkeleton ? *SharingSkeleton->GetName() : TEXT("None"));
 		return ArrayIndex;
 	}();
 
@@ -445,6 +444,10 @@ void UAnimationSharingManager::RegisterActorWithSkeleton(AActor* InActor, const
 			ActorData.UpdateActorHandleDelegate.ExecuteIfBound(ActorHandle);
 		}
 	}
+	else
+	{
+		UE_LOG(LogAnimationSharing, Error, TEXT("Invalid skeleton (%s) for which there is no sharing setup available!"), SharingSkeleton ? *SharingSkeleton->GetName() : TEXT("None"));
+	}
 }
 
 void UAnimationSharingManager::RegisterActorWithSkeletonBP(AActor* InActor, const USkeleton* SharingSkeleton)
@@ -768,7 +771,7 @@ void UAnimSharingInstance::Setup(UAnimationSharingManager* AnimationSharingManag
 			const uint8 StateValue = StateEntry.State;
 			const uint32 StateIndex = StateEnum->GetIndexByValue(StateValue);
 
-			checkf(!PerStateData.FindByPredicate([&](const FPerStateData& State) { return State.StateEnumValue == StateValue; }), TEXT("State already defined"));
+			checkf(!PerStateData.FindByPredicate([StateValue](const FPerStateData& State) { return State.StateEnumValue == StateValue; }), TEXT("State already defined"));
 
 			FPerStateData& StateData = PerStateData[StateIndex];
 			StateData.StateEnumValue = StateValue;
@@ -1019,7 +1022,7 @@ void UAnimSharingInstance::TickDebugInformation()
 			const FPerActorData& ActorData = PerActorData[ActorIndex];
 			const uint8 State = ActorData.CurrentState;
 
-			const FString StateString = [&]() -> FString
+			const FString StateString = [this, &ActorData, State]() -> FString
 			{
 				/** Check whether or not we are currently blending */
 				const uint32 BlendInstanceIndex = ActorData.BlendInstanceIndex;
@@ -1040,7 +1043,7 @@ void UAnimSharingInstance::TickDebugInformation()
 				return FString::Printf(TEXT("State - %s %1.2f"), *StateEnum->GetDisplayNameTextByValue(State).ToString(), ActorData.SignificanceValue);
 			}();
 
-			const FColor DebugColor = [&]()
+			const FColor DebugColor = [&ActorData]()
 			{
 				const uint32 BlendInstanceIndex = ActorData.BlendInstanceIndex;
 				const uint32 DemandInstanceIndex = ActorData.OnDemandInstanceIndex;
@@ -1110,7 +1113,7 @@ void UAnimSharingInstance::TickOnDemandInstances()
 #if LOG_STATES 
 			UE_LOG(LogAnimationSharing, Log, TEXT("Finished on demand %s"), *StateEnum->GetDisplayNameTextByValue(Instance.State).ToString());
 #endif
-			auto SetActorState = [&](uint32 ActorIndex, uint8 NewState)
+			auto SetActorState = [this, &Instance](uint32 ActorIndex, uint8 NewState)
 			{
 				if (Instance.BlendToPermutationIndex != INDEX_NONE)
 				{
@@ -1199,7 +1202,7 @@ void UAnimSharingInstance::TickOnDemandInstances()
 				const bool bShouldBlend = ScalabilitySettings->UseBlendTransitions.Default && PerActorData[ActorIndex].SignificanceValue >= ScalabilitySettings->BlendSignificanceValue.Default;
 
 				// Determine state to blend to
-				const uint8 BlendToState = [&]() -> uint8
+				const uint8 BlendToState = [this, &Instance, bShouldBlend, ActorIndex]() -> uint8
 				{
 					if (bShouldBlend)
 					{
@@ -1324,7 +1327,7 @@ void UAnimSharingInstance::TickActorStates()
 				/** When we are currently running an on-demand state we do not want as state change to impact the current animation */
 				const bool bShouldNotProcess = ActorData.bRunningOnDemand && !PerStateData[CurrentState].bIsOnDemand;
 
-				auto UpdateState = [&]()
+				auto UpdateState = [&ActorData, ActorIndex, CurrentState, PreviousState]()
 				{
 #if LOG_STATES 
 					UE_LOG(LogAnimationSharing, Log, TEXT("Setting %i state to %i previous %i | %i"), ActorIndex, CurrentState, PreviousState, ActorData.PermutationIndex);
@@ -1758,7 +1761,7 @@ uint32 UAnimSharingInstance::SetupBlend(uint8 FromState, uint8 ToState, uint32 A
 	uint32 BlendInstanceIndex = INDEX_NONE;
 	if (!bConcurrentBlendsReached)
 	{
-		BlendInstanceIndex = BlendInstances.IndexOfByPredicate([&](const FBlendInstance& Instance)
+		BlendInstanceIndex = BlendInstances.IndexOfByPredicate([this, FromState, ToState, bOnDemand, ActorIndex](const FBlendInstance& Instance)
 		{			
 			return (!Instance.bActive &&				// The instance should not have started yet
 				Instance.StateFrom == FromState &&		// It should be blending from the same state
