Fixed unsigned overflow of NumRandomizedInstances for animation. Negative values are now bound to 1.

diff --git a/Engine/Plugins/Developer/AnimationSharing/Source/AnimationSharing/Private/AnimationSharingManager.cpp b/Engine/Plugins/Developer/AnimationSharing/Source/AnimationSharing/Private/AnimationSharingManager.cpp
--- a/Engine/Plugins/Developer/AnimationSharing/Source/AnimationSharing/Private/AnimationSharingManager.cpp
+++ b/Engine/Plugins/Developer/AnimationSharing/Source/AnimationSharing/Private/AnimationSharingManager.cpp
@@ -874,7 +874,7 @@ void UAnimSharingInstance::SetupState(FPerStateData& StateData, const FAnimation
 #if WITH_EDITOR			
 			NumRandomizedInstances = AnimationSetup.NumRandomizedInstances.GetValueForPlatformIdentifiers(PlatformName, PlatformName);
 #endif
-			const uint32 NumInstances = StateEntry.bOnDemand ? NumInstancesPerSetup	: NumRandomizedInstances;
+			const uint32 NumInstances = StateEntry.bOnDemand ? NumInstancesPerSetup	: FGenericPlatformMath::Max(NumRandomizedInstances, 1);
 			for (uint32 InstanceIndex = 0; InstanceIndex < NumInstances; ++InstanceIndex)
 			{
 				if (!StateData.bIsAdditive)
