UE-70878
Crash starting PIE when two Animation Sharing states define the same state

diff --git a/Engine/Plugins/Developer/AnimationSharing/Source/AnimationSharing/Private/AnimationSharingManager.cpp b/Engine/Plugins/Developer/AnimationSharing/Source/AnimationSharing/Private/AnimationSharingManager.cpp
--- a/Engine/Plugins/Developer/AnimationSharing/Source/AnimationSharing/Private/AnimationSharingManager.cpp
+++ b/Engine/Plugins/Developer/AnimationSharing/Source/AnimationSharing/Private/AnimationSharingManager.cpp
@@ -771,11 +771,16 @@ void UAnimSharingInstance::Setup(UAnimationSharingManager* AnimationSharingManag
 			const uint8 StateValue = StateEntry.State;
 			const uint32 StateIndex = StateEnum->GetIndexByValue(StateValue);
 
-			checkf(!PerStateData.FindByPredicate([StateValue](const FPerStateData& State) { return State.StateEnumValue == StateValue; }), TEXT("State already defined"));
-
-			FPerStateData& StateData = PerStateData[StateIndex];
-			StateData.StateEnumValue = StateValue;
-			SetupState(StateData, StateEntry, SkeletalMesh, SkeletonSetup, Index);
+			if (!PerStateData.FindByPredicate([StateValue](const FPerStateData& State) { return State.StateEnumValue == StateValue; }))
+			{
+				FPerStateData& StateData = PerStateData[StateIndex];
+				StateData.StateEnumValue = StateValue;
+				SetupState(StateData, StateEntry, SkeletalMesh, SkeletonSetup, Index);
+			}
+			else
+			{
+				UE_LOG(LogAnimationSharing, Error, TEXT("Duplicate entries in Animation Setup for State %s"), *StateEnum->GetDisplayNameTextByValue(StateValue).ToString());
+			}
 		}
 
 		/** Setup blend actors, if enabled*/
