Crash entering PIE when AnimationSharingSetup doesn't have a valid/enabled state that is trying to be played - Assertion failed: Components.IsValidIndex(PermutationIndex)
#fix Change Check to log error to inform user about invalid setup

diff --git a/Engine/Plugins/Developer/AnimationSharing/Source/AnimationSharing/Private/AnimationSharingManager.cpp b/Engine/Plugins/Developer/AnimationSharing/Source/AnimationSharing/Private/AnimationSharingManager.cpp
--- a/Engine/Plugins/Developer/AnimationSharing/Source/AnimationSharing/Private/AnimationSharingManager.cpp
+++ b/Engine/Plugins/Developer/AnimationSharing/Source/AnimationSharing/Private/AnimationSharingManager.cpp
@@ -1704,6 +1704,13 @@ void UAnimSharingInstance::SetMasterComponentForActor(uint32 ActorIndex, USkelet
 void UAnimSharingInstance::SetupSlaveComponent(uint8 CurrentState, uint32 ActorIndex)
 {
 	const FPerStateData& StateData = PerStateData[CurrentState];
+
+	if (StateData.Components.Num() == 0)
+	{	
+		UE_LOG(LogAnimationSharing, Warning, TEXT("No Master Components available for state %s, make sure to set up an Animation Sequence/Blueprint "), *StateEnum->GetDisplayNameTextByValue(CurrentState).ToString());
+		return;
+	}
+
 	if (!StateData.bIsOnDemand)
 	{
 		const uint32 PermutationIndex = DeterminePermutationIndex(ActorIndex, CurrentState);
