Crash when returning to a previous state in an Animation Sharing Setup - Assertion failed: Instance.bActive
#fix Make sure on demand state is initialized when setting up slave component after finish an existing on demand state.

diff --git a/Engine/Plugins/Developer/AnimationSharing/Source/AnimationSharing/Private/AnimationSharingManager.cpp b/Engine/Plugins/Developer/AnimationSharing/Source/AnimationSharing/Private/AnimationSharingManager.cpp
--- a/Engine/Plugins/Developer/AnimationSharing/Source/AnimationSharing/Private/AnimationSharingManager.cpp
+++ b/Engine/Plugins/Developer/AnimationSharing/Source/AnimationSharing/Private/AnimationSharingManager.cpp
@@ -1127,6 +1127,21 @@ void UAnimSharingInstance::TickOnDemandInstances()
 				else
 				{
 					SetupSlaveComponent(NewState, ActorIndex);
+
+					// If we are setting up a slave to an on-demand state that is not in use yet it needs to create a new On Demand Instance which will not be kicked-off yet, so do that directly.
+					if (PerStateData[NewState].bIsOnDemand)
+					{
+						const int32 OnDemandInstanceIndex = PerActorData[ActorIndex].OnDemandInstanceIndex;
+						if (OnDemandInstanceIndex != INDEX_NONE)
+						{
+							FOnDemandInstance& NewOnDemandInstance = OnDemandInstances[OnDemandInstanceIndex];
+							if (!NewOnDemandInstance.bActive)
+							{
+								NewOnDemandInstance.bActive = true;
+								NewOnDemandInstance.StartTime = WorldTime;
+							}
+						}
+					}
 				}
 
 				// Set actor states
