Crashing running out of memory when entering pie when Num Random Instances is set to a negative value - Error: Direct3DDevice->CreateShaderResourceView failed - Anim Sharing Setup variables can be set to negative values
#fix Push down metadata to PerPlatformProperties when displayed in the details view

Skeletal Mesh snaps to its reference pose when blending from one animation to another while using an Animation Sharing setup
#fix Change default value for blend significance to 0.0f

diff --git a/Engine/Plugins/Developer/AnimationSharing/Source/AnimationSharing/Public/AnimationSharingTypes.h b/Engine/Plugins/Developer/AnimationSharing/Source/AnimationSharing/Public/AnimationSharingTypes.h
--- a/Engine/Plugins/Developer/AnimationSharing/Source/AnimationSharing/Public/AnimationSharingTypes.h
+++ b/Engine/Plugins/Developer/AnimationSharing/Source/AnimationSharing/Public/AnimationSharingTypes.h
@@ -29,7 +29,7 @@ public:
 	TSubclassOf<UAnimSharingStateInstance> AnimBlueprint;	
 
 	/** The number of randomized instances created from this setup, it will create instance with different start time offsets (Length / Number of Instance) * InstanceIndex */
-	UPROPERTY(EditAnywhere, Category = AnimationSharing)
+	UPROPERTY(EditAnywhere, Category = AnimationSharing, meta = (ClampMin = "1", UIMin = "1"))
 	FPerPlatformInt NumRandomizedInstances;
 
 	/** Whether or not this setup is enabled for specific platforms */
@@ -61,7 +61,7 @@ public:
 	bool bAdditive;
 
 	/** Duration of blending when blending to this state */
-	UPROPERTY(EditAnywhere, Category = AnimationSharing, meta = (EditCondition = "!bAdditive"))
+	UPROPERTY(EditAnywhere, Category = AnimationSharing, meta = (EditCondition = "!bAdditive", UIMin = "0.0", ClampMin = "0.0"))
 	float BlendTime;
 
 	/** Flag whether or not we should return to the previous state, only used when this state is an on-demand one*/
@@ -76,12 +76,12 @@ public:
 	uint8 NextState;
 
 	/** Number of instances that will be created for this state (platform-specific) */
-	UPROPERTY(EditAnywhere, Category = AnimationSharing, meta=(EditCondition="bOnDemand"))
+	UPROPERTY(EditAnywhere, Category = AnimationSharing, meta=(EditCondition="bOnDemand", ClampMin = "0", UIMin = "0"))
 	FPerPlatformInt MaximumNumberOfConcurrentInstances;
 
 	/** Percentage of 'wiggle' frames, this is used when we run out of available entries in Components, if one of the on-demand animations has started SequenceLength * WiggleFramePercentage ago or earlier,
 	it is used instead of a brand new one */
-	UPROPERTY(EditAnywhere, Category = AnimationSharing, meta = (EditCondition = "bOnDemand", UIMin="0.0", UIMax="1.0"))
+	UPROPERTY(EditAnywhere, Category = AnimationSharing, meta = (EditCondition = "bOnDemand", UIMin="0.0", UIMax="1.0", ClampMin = "0.0", ClampMax = "1.0"))
 	float WiggleTimePercentage;
 	
 	/** Whether or not this animation requires curves or morphtargets to function correctly for slave components */
@@ -160,22 +160,22 @@ struct FAnimationSharingScalability
 {
 	GENERATED_BODY()
 public:
-	FAnimationSharingScalability() : UseBlendTransitions(true), BlendSignificanceValue(0.5f), MaximumNumberConcurrentBlends(1) {}
+	FAnimationSharingScalability() : UseBlendTransitions(true), BlendSignificanceValue(0.0f), MaximumNumberConcurrentBlends(1) {}
 
 	/** Flag whether or not to use blend transitions between states */
 	UPROPERTY(EditAnywhere, Category = AnimationSharing)
 	FPerPlatformBool UseBlendTransitions;
 
 	/** Significance value tied to whether or not a transition should be blended */
-	UPROPERTY(EditAnywhere, Category = AnimationSharing, meta = (EditCondition = "bBlendTransitions"))
+	UPROPERTY(EditAnywhere, Category = AnimationSharing, meta = (EditCondition = "bBlendTransitions", ClampMin="0.0", UIMin="0.0"))
 	FPerPlatformFloat BlendSignificanceValue;
 
 	/** Maximum number of blends which can be running concurrently */
-	UPROPERTY(EditAnywhere, Category = AnimationSharing, meta = (EditCondition = "bBlendTransitions"))
+	UPROPERTY(EditAnywhere, Category = AnimationSharing, meta = (EditCondition = "bBlendTransitions", ClampMin = "1", UIMin = "1"))
 	FPerPlatformInt MaximumNumberConcurrentBlends;
 
 	/** Significance value tied to whether or not the master pose components should be ticking */
-	UPROPERTY(EditAnywhere, Category = AnimationSharing, meta = (EditCondition = "bBlendTransitions"))
+	UPROPERTY(EditAnywhere, Category = AnimationSharing, meta = (EditCondition = "bBlendTransitions", ClampMin = "0.0", UIMin = "0.0"))
 	FPerPlatformFloat TickSignificanceValue;
 };
 
diff --git a/Engine/Source/Editor/DetailCustomizations/Private/PerPlatformPropertyCustomization.cpp b/Engine/Source/Editor/DetailCustomizations/Private/PerPlatformPropertyCustomization.cpp
--- a/Engine/Source/Editor/DetailCustomizations/Private/PerPlatformPropertyCustomization.cpp
+++ b/Engine/Source/Editor/DetailCustomizations/Private/PerPlatformPropertyCustomization.cpp
@@ -16,6 +16,7 @@
 #include "SPerPlatformPropertiesWidget.h"
 #include "ScopedTransaction.h"
 #include "IPropertyUtilities.h"
+#include "UObject/MetaData.h"
 
 #define LOCTEXT_NAMESPACE "PerPlatformPropertyCustomization"
 
@@ -76,6 +77,21 @@ TSharedRef<SWidget> FPerPlatformPropertyCustomization<PerPlatformType>::GetWidge
 				}
 			}
 		}
+	
+	}
+
+	// Push down struct metadata to per-platform properties
+	{
+		// First get the source map
+		const TMap<FName, FString>* SourceMap = UMetaData::GetMapForObject(StructPropertyHandle->GetMetaDataProperty());
+		if (SourceMap)
+		{
+			// Iterate through source map, setting each key/value pair in the destination
+			for (const auto& It : *SourceMap)
+			{
+				EditProperty->SetInstanceMetaData(*It.Key.ToString(), *It.Value);
+			}
+		}
 	}
 
 	if (EditProperty.IsValid())
