- Crash entering PIE when first Anim Sharing state is disabled - Assertion failed: (Index >= 0) & (Index < ArrayNum)
- Crash during PIE when two Anim Sharing states define the same state - Assertion failed: Components.IsValidIndex(PermutationIndex)
- Ensure entering PIE when Anim Sharing state has no Anim Sequence

#fix Made initialization stage for Animation Sharing more error-resistant and outright fail to set up if we encounter any errors

diff --git a/Engine/Plugins/Developer/AnimationSharing/Source/AnimationSharing/Private/AnimationSharingManager.cpp b/Engine/Plugins/Developer/AnimationSharing/Source/AnimationSharing/Private/AnimationSharingManager.cpp
--- a/Engine/Plugins/Developer/AnimationSharing/Source/AnimationSharing/Private/AnimationSharingManager.cpp
+++ b/Engine/Plugins/Developer/AnimationSharing/Source/AnimationSharing/Private/AnimationSharingManager.cpp
@@ -226,9 +226,17 @@ void UAnimationSharingManager::SetupPerSkeletonData(const FPerSkeletonAnimationS
 	if (Skeleton && StateEnum && Processor)
 	{
 		UAnimSharingInstance* Data = NewObject<UAnimSharingInstance>(this);
-		PerSkeletonData.Add(Data);
-		Skeletons.Add(Skeleton);
-		Data->Setup(this, SkeletonSetup, &ScalabilitySettings, Skeletons.Num() - 1);
+		// Try and setup up instance using provided setup data
+		if (Data->Setup(this, SkeletonSetup, &ScalabilitySettings, Skeletons.Num() - 1))
+		{
+			PerSkeletonData.Add(Data);
+			Skeletons.Add(Skeleton);
+		}
+		else
+		{
+			UE_LOG(LogAnimationSharing, Error, TEXT("Failed to initialise Animation Sharing Data for Skeleton (%s)!"),
+				Skeleton ? *Skeleton->GetName() : TEXT("None"));
+		}
 	}
 	else
 	{
@@ -739,7 +747,7 @@ uint8 UAnimSharingInstance::DetermineStateForActor(uint32 ActorIndex, bool& bSho
 	return FMath::Max(0, State);
 }
 
-void UAnimSharingInstance::Setup(UAnimationSharingManager* AnimationSharingManager, const FPerSkeletonAnimationSharingSetup& SkeletonSetup, const FAnimationSharingScalability* InScalabilitySettings, uint32 Index)
+bool UAnimSharingInstance::Setup(UAnimationSharingManager* AnimationSharingManager, const FPerSkeletonAnimationSharingSetup& SkeletonSetup, const FAnimationSharingScalability* InScalabilitySettings, uint32 Index)
 {
 	USkeletalMesh* SkeletalMesh = SkeletonSetup.SkeletalMesh.LoadSynchronous();
 	/** Retrieve the state processor to use */
@@ -749,6 +757,8 @@ void UAnimSharingInstance::Setup(UAnimationSharingManager* AnimationSharingManag
 		bNativeStateProcessor = SkeletonSetup.StateProcessorClass->HasAnyClassFlags(CLASS_Native);
 	}
 
+	bool bErrors = false;
+
 	if (SkeletalMesh && StateProcessor)
 	{
 		SkeletalMeshBounds = SkeletalMesh->GetBounds().BoxExtent * 2;
@@ -776,15 +786,28 @@ void UAnimSharingInstance::Setup(UAnimationSharingManager* AnimationSharingManag
 				FPerStateData& StateData = PerStateData[StateIndex];
 				StateData.StateEnumValue = StateValue;
 				SetupState(StateData, StateEntry, SkeletalMesh, SkeletonSetup, Index);
+
+				// Make sure we have at least one component set up
+				if (StateData.Components.Num() == 0)
+				{
+					UE_LOG(LogAnimationSharing, Error, TEXT("No Components available for State %s"), *StateEnum->GetDisplayNameTextByValue(StateValue).ToString());
+					bErrors = true;
+				}
 			}
 			else
 			{
 				UE_LOG(LogAnimationSharing, Error, TEXT("Duplicate entries in Animation Setup for State %s"), *StateEnum->GetDisplayNameTextByValue(StateValue).ToString());
+				bErrors = true;
 			}
 		}
 
+		if (bErrors)
+		{
+			PerStateData.Empty();
+		}
+
 		/** Setup blend actors, if enabled*/
-		if (ScalabilitySettings->UseBlendTransitions.Default)
+		if (!bErrors && ScalabilitySettings->UseBlendTransitions.Default)
 		{
 			const uint32 TotalNumberOfBlendActorsRequired = ScalabilitySettings->MaximumNumberConcurrentBlends.Default;
 			const float ZOffset = Index * SkeletalMeshBounds.Z * 2.f;
@@ -809,7 +832,10 @@ void UAnimSharingInstance::Setup(UAnimationSharingManager* AnimationSharingManag
 	else
 	{
 		UE_LOG(LogAnimationSharing, Error, TEXT("Invalid Skeletal Mesh or State Processing Class"));
+		bErrors = true;
 	}
+
+	return !bErrors;
 }
 
 void UAnimSharingInstance::SetupState(FPerStateData& StateData, const FAnimationStateEntry& StateEntry, USkeletalMesh* SkeletalMesh, const FPerSkeletonAnimationSharingSetup& SkeletonSetup, uint32 Index)
@@ -865,7 +891,12 @@ void UAnimSharingInstance::SetupState(FPerStateData& StateData, const FAnimation
 		/** User can setup either an AnimBP or AnimationSequence */
 		UClass* AnimBPClass = AnimationSetup.AnimBlueprint.Get();
 		UAnimSequence* AnimSequence = AnimationSetup.AnimSequence.LoadSynchronous();
-		ensureMsgf(AnimBPClass != nullptr || AnimSequence != nullptr, TEXT("Animation setup without either an Animation Blueprint Class of Animation Sequence"));
+		
+		if (AnimBPClass == nullptr && AnimSequence == nullptr)
+		{
+			UE_LOG(LogAnimationSharing, Error, TEXT("Animation setup entry for state %s without either a valid Animation Blueprint Class or Animation Sequence"), StateEnum ? *StateEnum->GetName() : TEXT("None"));
+			continue;
+		}
 
 		bool bEnabled = AnimationSetup.Enabled.Default;
 #if WITH_EDITOR			
diff --git a/Engine/Plugins/Developer/AnimationSharing/Source/AnimationSharing/Public/AnimationSharingManager.h b/Engine/Plugins/Developer/AnimationSharing/Source/AnimationSharing/Public/AnimationSharingManager.h
--- a/Engine/Plugins/Developer/AnimationSharing/Source/AnimationSharing/Public/AnimationSharingManager.h
+++ b/Engine/Plugins/Developer/AnimationSharing/Source/AnimationSharing/Public/AnimationSharingManager.h
@@ -275,7 +275,7 @@ public:
 	uint8 DetermineStateForActor(uint32 ActorIndex, bool& bShouldProcess);
 
 	/** Initial set up of all animation sharing data and states */
-	void Setup(UAnimationSharingManager* AnimationSharingManager, const FPerSkeletonAnimationSharingSetup& SkeletonSetup, const FAnimationSharingScalability* ScalabilitySettings, uint32 Index);
+	bool Setup(UAnimationSharingManager* AnimationSharingManager, const FPerSkeletonAnimationSharingSetup& SkeletonSetup, const FAnimationSharingScalability* ScalabilitySettings, uint32 Index);
 	/** Populates data for a state setup */
 	void SetupState(FPerStateData& StateData, const FAnimationStateEntry& StateEntry, USkeletalMesh* SkeletalMesh, const FPerSkeletonAnimationSharingSetup& SkeletonSetup, uint32 Index);
 
