//UE4/Release-4.22 - Run Incremental Automated Tests Win64 - Screenshot 'AnimationSharing_Debug' test failed, Screenshots were different

diff --git a/Engine/Plugins/Developer/AnimationSharing/Source/AnimationSharing/Private/AnimationSharingManager.cpp b/Engine/Plugins/Developer/AnimationSharing/Source/AnimationSharing/Private/AnimationSharingManager.cpp
--- a/Engine/Plugins/Developer/AnimationSharing/Source/AnimationSharing/Private/AnimationSharingManager.cpp
+++ b/Engine/Plugins/Developer/AnimationSharing/Source/AnimationSharing/Private/AnimationSharingManager.cpp
@@ -227,7 +227,7 @@ void UAnimationSharingManager::SetupPerSkeletonData(const FPerSkeletonAnimationS
 	{
 		UAnimSharingInstance* Data = NewObject<UAnimSharingInstance>(this);
 		// Try and setup up instance using provided setup data
-		if (Data->Setup(this, SkeletonSetup, &ScalabilitySettings, Skeletons.Num() - 1))
+		if (Data->Setup(this, SkeletonSetup, &ScalabilitySettings, Skeletons.Num()))
 		{
 			PerSkeletonData.Add(Data);
 			Skeletons.Add(Skeleton);
