Make sure that additive sharing skeletal mesh components follow the same path as others

diff --git a/Engine/Plugins/Developer/AnimationSharing/Source/AnimationSharing/Private/AnimationSharingManager.cpp b/Engine/Plugins/Developer/AnimationSharing/Source/AnimationSharing/Private/AnimationSharingManager.cpp
--- a/Engine/Plugins/Developer/AnimationSharing/Source/AnimationSharing/Private/AnimationSharingManager.cpp
+++ b/Engine/Plugins/Developer/AnimationSharing/Source/AnimationSharing/Private/AnimationSharingManager.cpp
@@ -980,7 +980,7 @@ void UAnimSharingInstance::SetupState(FPerStateData& StateData, const FAnimation
 					Components.Add(Component);
 				}
 				else
-				{					
+				{	
 					const FVector SpawnLocation(FVector(NumSetups * SkeletalMeshBounds.X, 0.f, ZOffset));
 					const FName AdditiveComponentName(*(SkeletalMesh->GetName() + TEXT("_") + StateEnum->GetNameStringByIndex(StateEntry.State) + FString::FromInt(InstanceIndex)));
 					USkeletalMeshComponent* AdditiveComponent = NewObject<USkeletalMeshComponent>(SharingActor, AdditiveComponentName);
@@ -994,6 +994,10 @@ void UAnimSharingInstance::SetupState(FPerStateData& StateData, const FAnimation
 					FAdditiveAnimationInstance* AdditiveInstance = new FAdditiveAnimationInstance();
 					AdditiveInstance->Initialise(AdditiveComponent, SkeletonSetup.AdditiveAnimBlueprint.Get());
 					AdditiveInstanceStack.AddInstance(AdditiveInstance);
+					
+					/** Set the current animation length length */
+					StateData.AnimationLengths.Add(AnimSequence->SequenceLength);
+					Components.Add(AdditiveComponent);
 				}
 				
 				++NumSetups;
@@ -1309,6 +1313,9 @@ void UAnimSharingInstance::TickAdditiveInstances()
 	for (int32 InstanceIndex = 0; InstanceIndex < AdditiveInstances.Num(); ++InstanceIndex)
 	{
 		FAdditiveInstance& Instance = AdditiveInstances[InstanceIndex];
+		SetComponentUsage(true, Instance.State, Instance.UsedPerStateComponentIndex);		
+		SetComponentTick(Instance.State, Instance.UsedPerStateComponentIndex);
+
 		if (Instance.bActive)
 		{
 			const float WorldTimeSeconds = GetWorld()->GetTimeSeconds();
@@ -1326,6 +1333,8 @@ void UAnimSharingInstance::TickAdditiveInstances()
 				FreeAdditiveInstance(Instance.AdditiveAnimationInstance);
 				RemoveAdditiveInstance(InstanceIndex);
 				--InstanceIndex;
+
+				SetComponentUsage(false, Instance.State, Instance.UsedPerStateComponentIndex);
 			}
 		}
 		else
@@ -2025,6 +2034,7 @@ uint32 UAnimSharingInstance::SetupAdditiveInstance(uint8 StateIndex, uint8 FromS
 		const float WorldTimeSeconds = GetWorld()->GetTimeSeconds();
 		Instance.EndTime = WorldTimeSeconds + StateData.AdditiveAnimationSequence->SequenceLength;
 		Instance.State = StateIndex;
+		Instance.UsedPerStateComponentIndex = PerStateData[StateIndex].Components.IndexOfByKey(AnimationInstance->GetComponent());
 
 		InstanceIndex = AdditiveInstances.Num() - 1;
 		AnimationInstance->Setup(Instance.BaseComponent, StateData.AdditiveAnimationSequence);
diff --git a/Engine/Plugins/Developer/AnimationSharing/Source/AnimationSharing/Public/AnimationSharingManager.h b/Engine/Plugins/Developer/AnimationSharing/Source/AnimationSharing/Public/AnimationSharingManager.h
--- a/Engine/Plugins/Developer/AnimationSharing/Source/AnimationSharing/Public/AnimationSharingManager.h
+++ b/Engine/Plugins/Developer/AnimationSharing/Source/AnimationSharing/Public/AnimationSharingManager.h
@@ -112,6 +112,9 @@ struct FAdditiveInstance
 
 	/** Actor used for playing the additive animation */
 	FAdditiveAnimationInstance* AdditiveAnimationInstance;
+
+	/** Index into Components array for the current state data which is used for playing the animation*/
+	uint32 UsedPerStateComponentIndex;
 };
 
 /** Structure which holds data about a unique state which is linked to an enumeration value defined by the user. The data is populated from the user exposed FAnimationStateEntry */
