Additive animation in sharing system are not playing back correctly and fail to play when triggered in consecutive frames
#fix Make sure conditional statement does not fall through when trying to play an additive state while already in one
#fix Make sure we mark additive and slave components as ticking in the current frame whenever we setup an additive sharing instance rather than in the next frame during TickAdditiveInstances()
#misc Make debug text more useful to actually see corresponding state names as well

diff --git a/Engine/Plugins/Developer/AnimationSharing/Source/AnimationSharing/Private/AnimationSharingManager.cpp b/Engine/Plugins/Developer/AnimationSharing/Source/AnimationSharing/Private/AnimationSharingManager.cpp
--- a/Engine/Plugins/Developer/AnimationSharing/Source/AnimationSharing/Private/AnimationSharingManager.cpp
+++ b/Engine/Plugins/Developer/AnimationSharing/Source/AnimationSharing/Private/AnimationSharingManager.cpp
@@ -1045,7 +1045,7 @@ void UAnimSharingInstance::TickDebugInformation()
 		{
 			for (int32 Index = 0; Index < StateData.InUseComponentFrameBits.Num(); ++Index)
 			{
-				const FString ComponentString = FString::Printf(TEXT("In Use %s - Required %s"), StateData.InUseComponentFrameBits[Index] ? TEXT("True") : TEXT("False"), StateData.SlaveTickRequiredFrameBits[Index] ? TEXT("True") : TEXT("False"));
+				const FString ComponentString = FString::Printf(TEXT("[%s]\n In Use %s - Required %s"), *StateEnum->GetDisplayNameTextByValue(StateData.StateEnumValue).ToString(), StateData.InUseComponentFrameBits[Index] ? TEXT("True") : TEXT("False"), StateData.SlaveTickRequiredFrameBits[Index] ? TEXT("True") : TEXT("False"));
 				DrawDebugString(GetWorld(), StateData.Components[Index]->GetComponentLocation() + FVector(0,0,StateData.Components[Index]->Bounds.BoxExtent.Z), ComponentString, nullptr, FColor::White, 0.016f, false);
 			}
 		}
@@ -1405,18 +1405,21 @@ void UAnimSharingInstance::TickActorStates()
 #endif
 				}
 				/** Play additive animation only if actor isn't already playing one */
-				else if (PerStateData[CurrentState].bIsAdditive && !ActorData.bRunningAdditive)
-				{					
-					const uint32 AdditiveInstanceIndex = SetupAdditiveInstance(CurrentState, PreviousState, ActorData.PermutationIndex);
-					if (AdditiveInstanceIndex != INDEX_NONE)
+				else if (PerStateData[CurrentState].bIsAdditive)
+				{
+					if (!ActorData.bRunningAdditive)
 					{
-						ActorData.bRunningAdditive = true;
-						ActorData.AdditiveInstanceIndex = AdditiveInstanceIndex;
-						AdditiveInstances[AdditiveInstanceIndex].ActorIndex = ActorIndex;
+						const uint32 AdditiveInstanceIndex = SetupAdditiveInstance(CurrentState, PreviousState, ActorData.PermutationIndex);
+						if (AdditiveInstanceIndex != INDEX_NONE)
+						{
+							ActorData.bRunningAdditive = true;
+							ActorData.AdditiveInstanceIndex = AdditiveInstanceIndex;
+							AdditiveInstances[AdditiveInstanceIndex].ActorIndex = ActorIndex;
+						}
 					}
 				}
 				/** If we are _already_ running an on-demand instance and the new state is also an on-demand we'll have to blend the new state in*/
-				else if (PerStateData[CurrentState].bIsOnDemand)					
+				else if (PerStateData[CurrentState].bIsOnDemand)
 				{
 					/** If the new state is different than the currently running on-demand state, this could happen if we previously only updated the state and not processed it */
 					const bool bSetupInstance = (!ActorData.bRunningOnDemand || (ActorData.bRunningOnDemand && OnDemandInstances[ActorData.OnDemandInstanceIndex].State != CurrentState));
@@ -2038,6 +2041,9 @@ uint32 UAnimSharingInstance::SetupAdditiveInstance(uint8 StateIndex, uint8 FromS
 
 		InstanceIndex = AdditiveInstances.Num() - 1;
 		AnimationInstance->Setup(Instance.BaseComponent, StateData.AdditiveAnimationSequence);
+
+		SetComponentUsage(true, Instance.State, Instance.UsedPerStateComponentIndex);
+		SetComponentTick(Instance.State, Instance.UsedPerStateComponentIndex);
 	}
 
 	return InstanceIndex;
