Fix issue causing the user being unable to create a new Animation Sharing Setup asset

diff --git a/Engine/Plugins/Developer/AnimationSharing/Source/AnimationSharingEd/Public/AssetTypeActions_AnimationSharingSetup.h b/Engine/Plugins/Developer/AnimationSharing/Source/AnimationSharingEd/Public/AssetTypeActions_AnimationSharingSetup.h
--- a/Engine/Plugins/Developer/AnimationSharing/Source/AnimationSharingEd/Public/AssetTypeActions_AnimationSharingSetup.h
+++ b/Engine/Plugins/Developer/AnimationSharing/Source/AnimationSharingEd/Public/AssetTypeActions_AnimationSharingSetup.h
@@ -3,16 +3,16 @@
 #pragma once
 
 #include "CoreMinimal.h"
-#include "AnimationSharingManager.h"
 #include "AssetTypeCategories.h"
 #include "AssetTypeActions_Base.h"
+#include "AnimationSharingSetup.h"
 
 class FAssetTypeActions_AnimationSharingSetup : public FAssetTypeActions_Base
 {
 public:
 	// IAssetTypeActions Implementation
-	virtual FText GetName() const override { return NSLOCTEXT("AssetTypeActions", "AssetTypeActions_AnimationSharingSetup", "AnimationSharingSetup"); }
+	virtual FText GetName() const override { return NSLOCTEXT("AssetTypeActions", "AssetTypeActions_AnimationSharingSetup", "Animation Sharing Setup"); }
 	virtual FColor GetTypeColor() const override { return FColor(255, 10, 255); }
-	virtual UClass* GetSupportedClass() const override { return UAnimationSharingManager::StaticClass(); }	
+	virtual UClass* GetSupportedClass() const override { return UAnimationSharingSetup::StaticClass(); }	
 	virtual uint32 GetCategories() override { return EAssetTypeCategories::Animation; }
 };
