User is able to set a destructible mesh to a skeletal mesh asset property
#fix added DisallowedClasses metadata for UDestructibleMesh in various USkeletalMesh UProperties

diff --git a/Engine/Plugins/Developer/AnimationSharing/Source/AnimationSharing/Public/AnimationSharingTypes.h b/Engine/Plugins/Developer/AnimationSharing/Source/AnimationSharing/Public/AnimationSharingTypes.h
--- a/Engine/Plugins/Developer/AnimationSharing/Source/AnimationSharing/Public/AnimationSharingTypes.h
+++ b/Engine/Plugins/Developer/AnimationSharing/Source/AnimationSharing/Public/AnimationSharingTypes.h
@@ -135,7 +135,7 @@ public:
 	TObjectPtr<USkeleton> Skeleton;
 
 	/** Skeletal mesh used to setup skeletal mesh components */
-	UPROPERTY(EditAnywhere, Category = AnimationSharing)
+	UPROPERTY(EditAnywhere, Category = AnimationSharing, meta = (DisallowedClasses = "DestructibleMesh"))
 	TObjectPtr<USkeletalMesh> SkeletalMesh;
 
 	/** Animation blueprint used to perform the blending between states */
diff --git a/Engine/Plugins/FX/Niagara/Source/Niagara/Classes/NiagaraDataInterfaceSkeletalMesh.h b/Engine/Plugins/FX/Niagara/Source/Niagara/Classes/NiagaraDataInterfaceSkeletalMesh.h
--- a/Engine/Plugins/FX/Niagara/Source/Niagara/Classes/NiagaraDataInterfaceSkeletalMesh.h
+++ b/Engine/Plugins/FX/Niagara/Source/Niagara/Classes/NiagaraDataInterfaceSkeletalMesh.h
@@ -716,7 +716,7 @@ public:
 
 #if WITH_EDITORONLY_DATA
 	/** Mesh used to sample from when not overridden by a source actor from the scene. Only available in editor for previewing. This is removed in cooked builds. */
-	UPROPERTY(EditAnywhere, Category = "Mesh")
+	UPROPERTY(EditAnywhere, Category = "Mesh", meta = (DisallowedClasses = "DestructibleMesh"))
 	TSoftObjectPtr<USkeletalMesh> PreviewMesh;
 #endif
 
diff --git a/Engine/Source/Runtime/Engine/Classes/Animation/AnimSequence.h b/Engine/Source/Runtime/Engine/Classes/Animation/AnimSequence.h
--- a/Engine/Source/Runtime/Engine/Classes/Animation/AnimSequence.h
+++ b/Engine/Source/Runtime/Engine/Classes/Animation/AnimSequence.h
@@ -323,7 +323,7 @@ public:
 
 #if WITH_EDITORONLY_DATA
 	/** If RetargetSource is set to Default (None), this is asset for the base pose to use when retargeting. Transform data will be saved in RetargetSourceAssetReferencePose. */
-	UPROPERTY(EditAnywhere, AssetRegistrySearchable, Category=Animation)
+	UPROPERTY(EditAnywhere, AssetRegistrySearchable, Category=Animation, meta = (DisallowedClasses = "DestructibleMesh"))
 	TSoftObjectPtr<USkeletalMesh> RetargetSourceAsset;
 #endif
 
diff --git a/Engine/Source/Runtime/Engine/Classes/Animation/PoseAsset.h b/Engine/Source/Runtime/Engine/Classes/Animation/PoseAsset.h
--- a/Engine/Source/Runtime/Engine/Classes/Animation/PoseAsset.h
+++ b/Engine/Source/Runtime/Engine/Classes/Animation/PoseAsset.h
@@ -175,7 +175,7 @@ public:
 
 #if WITH_EDITORONLY_DATA
 	/** If RetargetSource is set to Default (None), this is asset for the base pose to use when retargeting. Transform data will be saved in RetargetSourceAssetReferencePose. */
-	UPROPERTY(EditAnywhere, AssetRegistrySearchable, Category=Animation)
+	UPROPERTY(EditAnywhere, AssetRegistrySearchable, Category=Animation, meta = (DisallowedClasses = "DestructibleMesh"))
 	TSoftObjectPtr<USkeletalMesh> RetargetSourceAsset;
 #endif
 
diff --git a/Engine/Source/Runtime/Engine/Classes/Components/SkinnedMeshComponent.h b/Engine/Source/Runtime/Engine/Classes/Components/SkinnedMeshComponent.h
--- a/Engine/Source/Runtime/Engine/Classes/Components/SkinnedMeshComponent.h
+++ b/Engine/Source/Runtime/Engine/Classes/Components/SkinnedMeshComponent.h
@@ -230,7 +230,7 @@ class ENGINE_API USkinnedMeshComponent : public UMeshComponent, public ILODSyncI
 	friend class FSkinnedMeshComponentRecreateRenderStateContext;
 
 	/** The skeletal mesh used by this component. */
-	UPROPERTY(EditAnywhere, BlueprintReadOnly, Category="Mesh")
+	UPROPERTY(EditAnywhere, BlueprintReadOnly, Category="Mesh", meta = (DisallowedClasses = "DestructibleMesh"))
 	TObjectPtr<class USkeletalMesh> SkeletalMesh;
 
 	//
diff --git a/Engine/Source/Runtime/Engine/Classes/Engine/PreviewMeshCollection.h b/Engine/Source/Runtime/Engine/Classes/Engine/PreviewMeshCollection.h
--- a/Engine/Source/Runtime/Engine/Classes/Engine/PreviewMeshCollection.h
+++ b/Engine/Source/Runtime/Engine/Classes/Engine/PreviewMeshCollection.h
@@ -30,7 +30,7 @@ struct FPreviewMeshCollectionEntry
 	}
 
 	/** The skeletal mesh to display */
-	UPROPERTY(EditAnywhere, Category = "Skeletal Mesh", meta=(DisplayThumbnail=true))
+	UPROPERTY(EditAnywhere, Category = "Skeletal Mesh", meta=(DisplayThumbnail=true, DisallowedClasses = "DestructibleMesh"))
 	TSoftObjectPtr<USkeletalMesh> SkeletalMesh;
 };
 
