Inclusive Terminology Push
Renamed / deprecated master/slave terms at AnimationSharing

diff --git a/Engine/Plugins/Developer/AnimationSharing/AnimationSharing.uplugin b/Engine/Plugins/Developer/AnimationSharing/AnimationSharing.uplugin
--- a/Engine/Plugins/Developer/AnimationSharing/AnimationSharing.uplugin
+++ b/Engine/Plugins/Developer/AnimationSharing/AnimationSharing.uplugin
@@ -3,7 +3,7 @@
 	"Version": 1,
 	"VersionName": "1.0",
 	"FriendlyName": "Animation Sharing",
-	"Description": "Plugin to create Shared Animation systems using the Master-Child pose functionality",
+	"Description": "Plugin to create Shared Animation systems using the Leader-Follower pose functionality",
 	"Category": "Animation",
 	"CreatedBy": "Epic Games, Inc.",
 	"CreatedByURL": "http://epicgames.com",
diff --git a/Engine/Plugins/Developer/AnimationSharing/Source/AnimationSharing/Private/AnimationSharingManager.cpp b/Engine/Plugins/Developer/AnimationSharing/Source/AnimationSharing/Private/AnimationSharingManager.cpp
--- a/Engine/Plugins/Developer/AnimationSharing/Source/AnimationSharing/Private/AnimationSharingManager.cpp
+++ b/Engine/Plugins/Developer/AnimationSharing/Source/AnimationSharing/Private/AnimationSharingManager.cpp
@@ -47,8 +47,8 @@ static FAutoConsoleVariableRef CVarAnimSharing_DebugStates(
 	TEXT("Values: 0/1/2/3\n")
 	TEXT("Controls whether and which animation sharing debug features are enabled.\n")
 	TEXT("0: Turned off.\n")
-	TEXT("1: Turns on active master-components and blend with material coloring, and printing state information for each actor above their capsule.\n")
-	TEXT("2: Turns printing state information about currently active animation states, blend etc. Also enables line drawing from slave-components to currently assigned master components."),
+	TEXT("1: Turns on active leader-components and blend with material coloring, and printing state information for each actor above their capsule.\n")
+	TEXT("2: Turns printing state information about currently active animation states, blend etc. Also enables line drawing from follower-components to currently assigned leader components."),
 	ECVF_Cheat);
 
 static int32 GAnimationSharingEnabled = 1;
@@ -78,27 +78,27 @@ static FAutoConsoleCommandWithWorldAndArgs CVarAnimSharing_Enabled(
 	ECVF_Cheat);
 
 #if !UE_BUILD_SHIPPING
-static int32 GMasterComponentsVisible = 0;
+static int32 GLeaderComponentsVisible = 0;
 static FAutoConsoleCommandWithWorldAndArgs CVarAnimSharing_ToggleVisibility(
 	TEXT("a.Sharing.ToggleVisibility"),
-	TEXT("Toggles the visibility of the Master Pose Components."),
+	TEXT("Toggles the visibility of the Leader Pose Components."),
 	FConsoleCommandWithWorldAndArgsDelegate::CreateLambda([](const TArray<FString>& Args, UWorld* World)
 	{
-		const bool bShouldBeVisible = !GMasterComponentsVisible;
+		const bool bShouldBeVisible = !GLeaderComponentsVisible;
 
 		/** Need to unregister actors here*/
 		UAnimationSharingManager* Manager = FAnimSharingModule::Get(World);
 		if (Manager)
 		{
-			Manager->SetMasterComponentsVisibility(bShouldBeVisible);
+			Manager->SetLeaderComponentsVisibility(bShouldBeVisible);
 		}
 
-		GMasterComponentsVisible = bShouldBeVisible;
+		GLeaderComponentsVisible = bShouldBeVisible;
 	}),
 	ECVF_Cheat);
 
 #else
-static const int32 GMasterComponentsVisible = 0;
+static const int32 GLeaderComponentsVisible = 0;
 #endif
 
 #if WITH_EDITOR
@@ -438,7 +438,7 @@ void UAnimationSharingManager::RegisterActorWithSkeleton(AActor* InActor, const
 				ActorData.ComponentIndices.Add(Data->PerComponentData.Num() - 1);
 
 				const int32 ComponentIndex = Data->PerComponentData.Num() - 1;
-				Data->SetupSlaveComponent(ActorData.CurrentState, ActorIndex);
+				Data->SetupFollowerComponent(ActorData.CurrentState, ActorIndex);
 			}
 
 			if (Data->PerStateData[ActorData.CurrentState].bIsOnDemand && ActorData.OnDemandInstanceIndex != INDEX_NONE)
@@ -611,7 +611,7 @@ void UAnimationSharingManager::UnregisterAllActors()
 	}	
 }
 
-void UAnimationSharingManager::SetMasterComponentsVisibility(bool bVisible)
+void UAnimationSharingManager::SetLeaderComponentsVisibility(bool bVisible)
 {
 	for (UAnimSharingInstance* Data : PerSkeletonData)
 	{
@@ -803,7 +803,7 @@ bool UAnimSharingInstance::Setup(UAnimationSharingManager* AnimationSharingManag
 				BlendComponent->RegisterComponent();
 				BlendComponent->SetRelativeLocation(SpawnLocation);
 				BlendComponent->SetSkeletalMesh(SkeletalMesh);
-				BlendComponent->SetVisibility(GMasterComponentsVisible == 1);
+				BlendComponent->SetVisibility(GLeaderComponentsVisible == 1);
 
 				BlendComponent->PrimaryComponentTick.AddPrerequisite(AnimSharingManager, AnimSharingManager->GetTickFunction());
 
@@ -908,7 +908,7 @@ void UAnimSharingInstance::SetupState(FPerStateData& StateData, const FAnimation
 					Component->SetSkeletalMesh(SkeletalMesh);
 					Component->VisibilityBasedAnimTickOption = EVisibilityBasedAnimTickOption::AlwaysTickPoseAndRefreshBones;
 					Component->SetForcedLOD(1);
-					Component->SetVisibility(GMasterComponentsVisible == 1);
+					Component->SetVisibility(GLeaderComponentsVisible == 1);
 					Component->bPropagateCurvesToFollowers = StateEntry.bRequiresCurves;
 
 					if (AnimBPClass != nullptr && AnimSequence != nullptr && !StateEntry.bOnDemand)
@@ -971,7 +971,7 @@ void UAnimSharingInstance::SetupState(FPerStateData& StateData, const FAnimation
 					AdditiveComponent->RegisterComponent();
 					AdditiveComponent->SetRelativeLocation(SpawnLocation);
 					AdditiveComponent->SetSkeletalMesh(SkeletalMesh);
-					AdditiveComponent->SetVisibility(GMasterComponentsVisible == 1);
+					AdditiveComponent->SetVisibility(GLeaderComponentsVisible == 1);
 
 					AdditiveComponent->PrimaryComponentTick.AddPrerequisite(AnimSharingManager, AnimSharingManager->GetTickFunction());
 
@@ -1016,20 +1016,20 @@ void UAnimSharingInstance::SetupState(FPerStateData& StateData, const FAnimation
 	/** This should enforce turning off the components tick during the first frame */
 	StateData.PreviousInUseComponentFrameBits.Init(true, Components.Num());
 
-	StateData.SlaveTickRequiredFrameBits.Init(false, Components.Num());
+	StateData.FollowerTickRequiredFrameBits.Init(false, Components.Num());
 }
 
 void UAnimSharingInstance::TickDebugInformation()
 {
 #if !UE_BUILD_SHIPPING
 #if UE_BUILD_DEVELOPMENT
-	if (GMasterComponentsVisible && GAnimationSharingDebugging >= 2)
+	if (GLeaderComponentsVisible && GAnimationSharingDebugging >= 2)
 	{
 		for (const FPerStateData& StateData : PerStateData)
 		{
 			for (int32 Index = 0; Index < StateData.InUseComponentFrameBits.Num(); ++Index)
 			{
-				const FString ComponentString = FString::Printf(TEXT("[%s]\n In Use %s - Required %s"), *StateEnum->GetDisplayNameTextByValue(StateData.StateEnumValue).ToString(), StateData.InUseComponentFrameBits[Index] ? TEXT("True") : TEXT("False"), StateData.SlaveTickRequiredFrameBits[Index] ? TEXT("True") : TEXT("False"));
+				const FString ComponentString = FString::Printf(TEXT("[%s]\n In Use %s - Required %s"), *StateEnum->GetDisplayNameTextByValue(StateData.StateEnumValue).ToString(), StateData.InUseComponentFrameBits[Index] ? TEXT("True") : TEXT("False"), StateData.FollowerTickRequiredFrameBits[Index] ? TEXT("True") : TEXT("False"));
 				DrawDebugString(GetWorld(), StateData.Components[Index]->GetComponentLocation() + FVector(0,0,StateData.Components[Index]->Bounds.BoxExtent.Z), ComponentString, nullptr, FColor::White, 0.016f, false);
 			}
 		}
@@ -1141,13 +1141,13 @@ void UAnimSharingInstance::TickOnDemandInstances()
 			{
 				if (Instance.BlendToPermutationIndex != INDEX_NONE)
 				{
-					SetPermutationSlaveComponent(NewState, ActorIndex, Instance.BlendToPermutationIndex);
+					SetPermutationFollowerComponent(NewState, ActorIndex, Instance.BlendToPermutationIndex);
 				}
 				else
 				{
-					SetupSlaveComponent(NewState, ActorIndex);
+					SetupFollowerComponent(NewState, ActorIndex);
 
-					// If we are setting up a slave to an on-demand state that is not in use yet it needs to create a new On Demand Instance which will not be kicked-off yet, so do that directly.
+					// If we are setting up a follower to an on-demand state that is not in use yet it needs to create a new On Demand Instance which will not be kicked-off yet, so do that directly.
 					if (PerStateData[NewState].bIsOnDemand)
 					{
 						const int32 OnDemandInstanceIndex = PerActorData[ActorIndex].OnDemandInstanceIndex;
@@ -1197,7 +1197,7 @@ void UAnimSharingInstance::TickOnDemandInstances()
 				// Only do this if the state is different than the current on-demand one
 				else if (CurrentState != Instance.State)
 				{
-					// If the new state is not an on-demand one and we are not currently blending, if we are blending the blend will set the final master component
+					// If the new state is not an on-demand one and we are not currently blending, if we are blending the blend will set the final leader component
 					if (!PerStateData[CurrentState].bIsOnDemand || !Instance.bBlendActive)
 					{
 						SetActorState(ActorIndex, CurrentState);
@@ -1312,7 +1312,7 @@ void UAnimSharingInstance::TickAdditiveInstances()
 					PerActorData[Instance.ActorIndex].AdditiveInstanceIndex = INDEX_NONE;
 
 					// Set it to base component on top of the additive animation is playing
-					SetMasterComponentForActor(Instance.ActorIndex, Instance.AdditiveAnimationInstance->GetBaseComponent());
+					SetLeaderComponentForActor(Instance.ActorIndex, Instance.AdditiveAnimationInstance->GetBaseComponent());
 				}
 				FreeAdditiveInstance(Instance.State, Instance.AdditiveAnimationInstance);
 				RemoveAdditiveInstance(InstanceIndex);
@@ -1327,7 +1327,7 @@ void UAnimSharingInstance::TickAdditiveInstances()
 			Instance.AdditiveAnimationInstance->Start();
 			if (Instance.ActorIndex != INDEX_NONE)
 			{				
-				SetMasterComponentForActor(Instance.ActorIndex, Instance.AdditiveAnimationInstance->GetComponent());
+				SetLeaderComponentForActor(Instance.ActorIndex, Instance.AdditiveAnimationInstance->GetComponent());
 			}
 		}
 	}
@@ -1468,7 +1468,7 @@ void UAnimSharingInstance::TickActorStates()
 					/** Otherwise just switch it to the new state */
 					if (bShouldSwitch)
 					{
-						SetupSlaveComponent(CurrentState, ActorIndex);
+						SetupFollowerComponent(CurrentState, ActorIndex);
 #if LOG_STATES 
 						UE_LOG(LogAnimationSharing, Log, TEXT("Changing state to %s from %s %i"), *StateEnum->GetDisplayNameTextByValue(CurrentState).ToString(), *StateEnum->GetDisplayNameTextByValue(PreviousState).ToString(), ActorIndex);
 #endif
@@ -1477,7 +1477,7 @@ void UAnimSharingInstance::TickActorStates()
 					UpdateState();
 				}
 			}
-			/** Flag the currently master component as in-use */
+			/** Flag the currently leader component as in-use */
 			else if (!ActorData.bRunningOnDemand && !ActorData.bBlending)
 			{
 #if LOG_STATES 
@@ -1494,7 +1494,7 @@ void UAnimSharingInstance::TickActorStates()
 				SetComponentUsage(true, ActorData.CurrentState, ActorData.PermutationIndex);
 			}
 			
-			// Propagate visibility to master component
+			// Propagate visibility to leader component
 			if (ActorData.bRequiresTick)
 			{
 				SetComponentTick(ActorData.CurrentState, ActorData.PermutationIndex);
@@ -1508,7 +1508,7 @@ void UAnimSharingInstance::RemoveFromCurrentBlend(int32 ActorIndex)
 	if (PerActorData[ActorIndex].bBlending && PerActorData[ActorIndex].BlendInstanceIndex != INDEX_NONE && BlendInstances.IsValidIndex(PerActorData[ActorIndex].BlendInstanceIndex))
 	{
 		FBlendInstance& OldBlendInstance = BlendInstances[PerActorData[ActorIndex].BlendInstanceIndex];
-		SetMasterComponentForActor(ActorIndex, OldBlendInstance.TransitionBlendInstance->GetToComponent());
+		SetLeaderComponentForActor(ActorIndex, OldBlendInstance.TransitionBlendInstance->GetToComponent());
 		OldBlendInstance.ActorIndices.Remove(ActorIndex);
 		PerActorData[ActorIndex].BlendInstanceIndex = INDEX_NONE;
 	}
@@ -1541,12 +1541,12 @@ void UAnimSharingInstance::TickBlendInstances()
 			const bool bToStateIsOnDemand = PerStateData[Instance.StateTo].bIsOnDemand;
 			const bool bFromStateIsOnDemand = PerStateData[Instance.StateFrom].bIsOnDemand;
 
-			// If we were blending to an on-demand state we need to set the on-demand component as the new master component
+			// If we were blending to an on-demand state we need to set the on-demand component as the new leader component
 			if (bToStateIsOnDemand)
 			{
 				for (uint32 ActorIndex : Instance.ActorIndices)
 				{
-					SetMasterComponentForActor(ActorIndex, Instance.TransitionBlendInstance->GetToComponent());
+					SetLeaderComponentForActor(ActorIndex, Instance.TransitionBlendInstance->GetToComponent());
 					PerActorData[ActorIndex].PermutationIndex = 0;
 #if LOG_STATES
 					UE_LOG(LogAnimationSharing, Log, TEXT("Setting %i to on-demand component %i"), ActorIndex, Instance.ToOnDemandInstanceIndex);
@@ -1558,8 +1558,8 @@ void UAnimSharingInstance::TickBlendInstances()
 					}
 				}
 			}
-			/** Otherwise if the state we were blending from was not on-demand we set the new state component as the new master component,
-				if we are blending from an on-demand state FOnDemandInstance with set the correct master component when it finishes	*/
+			/** Otherwise if the state we were blending from was not on-demand we set the new state component as the new leader component,
+				if we are blending from an on-demand state FOnDemandInstance with set the correct leader component when it finishes	*/
 			else if (!bFromStateIsOnDemand)
 			{				
 				for (uint32 ActorIndex : Instance.ActorIndices)
@@ -1569,7 +1569,7 @@ void UAnimSharingInstance::TickBlendInstances()
 #if LOG_STATES 
 						UE_LOG(LogAnimationSharing, Log, TEXT("Setting %i to state %i | %i"), ActorIndex, Instance.StateTo, Instance.ToPermutationIndex);
 #endif
-						SetPermutationSlaveComponent(Instance.StateTo, ActorIndex, Instance.ToPermutationIndex);
+						SetPermutationFollowerComponent(Instance.StateTo, ActorIndex, Instance.ToPermutationIndex);
 #if !UE_BUILD_SHIPPING
 						for (uint32 ComponentIndex : PerActorData[ActorIndex].ComponentIndices)
 						{
@@ -1597,12 +1597,12 @@ void UAnimSharingInstance::TickBlendInstances()
 		}
 		else
 		{
-			// Check whether or not the blend has started, if not set up the actors as slaves at this point
+			// Check whether or not the blend has started, if not set up the actors as followers at this point
 			if (!Instance.bBlendStarted)
 			{
 				for (uint32 ActorIndex : Instance.ActorIndices)
 				{
-					SetMasterComponentForActor(ActorIndex, Instance.TransitionBlendInstance->GetComponent());			
+					SetLeaderComponentForActor(ActorIndex, Instance.TransitionBlendInstance->GetComponent());			
 
 					for (uint32 ComponentIndex : PerActorData[ActorIndex].ComponentIndices)
 					{
@@ -1645,7 +1645,7 @@ void UAnimSharingInstance::TickAnimationStates()
 			const bool bPreviousState = StateData.PreviousInUseComponentFrameBits[Index];
 			const bool bCurrentState = StateData.InUseComponentFrameBits[Index];
 
-			const bool bShouldTick = StateData.SlaveTickRequiredFrameBits[Index];
+			const bool bShouldTick = StateData.FollowerTickRequiredFrameBits[Index];
 
 			if (bCurrentState != bPreviousState)
 			{
@@ -1676,7 +1676,7 @@ void UAnimSharingInstance::TickAnimationStates()
 		// Set previous to current and reset current bits
 		StateData.PreviousInUseComponentFrameBits = StateData.InUseComponentFrameBits;
 		StateData.InUseComponentFrameBits.Init(false, StateData.PreviousInUseComponentFrameBits.Num());
-		StateData.SlaveTickRequiredFrameBits.Init(false, StateData.SlaveTickRequiredFrameBits.Num());
+		StateData.FollowerTickRequiredFrameBits.Init(false, StateData.FollowerTickRequiredFrameBits.Num());
 
 		/** Reset on demand index for next frame */
 		StateData.CurrentFrameOnDemandIndex = INDEX_NONE;
@@ -1701,9 +1701,9 @@ void UAnimSharingInstance::SetComponentUsage(bool bUsage, uint8 StateIndex, uint
 
 void UAnimSharingInstance::SetComponentTick(uint8 StateIndex, uint32 ComponentIndex)
 {
-	if (PerStateData[StateIndex].SlaveTickRequiredFrameBits.IsValidIndex(ComponentIndex))
+	if (PerStateData[StateIndex].FollowerTickRequiredFrameBits.IsValidIndex(ComponentIndex))
 	{
-		PerStateData[StateIndex].SlaveTickRequiredFrameBits[ComponentIndex] = true;
+		PerStateData[StateIndex].FollowerTickRequiredFrameBits[ComponentIndex] = true;
 	}
 }
 
@@ -1720,7 +1720,7 @@ void UAnimSharingInstance::FreeAdditiveInstance(uint8 StateIndex, FAdditiveAnima
 	PerStateData[StateIndex].AdditiveInstanceStack.FreeInstance(Instance);
 }
 
-void UAnimSharingInstance::SetMasterComponentForActor(uint32 ActorIndex, USkeletalMeshComponent* Component)
+void UAnimSharingInstance::SetLeaderComponentForActor(uint32 ActorIndex, USkeletalMeshComponent* Component)
 {
 	// Always ensure the component is ticking
 	if (Component)
@@ -1744,20 +1744,20 @@ void UAnimSharingInstance::SetMasterComponentForActor(uint32 ActorIndex, USkelet
 	}
 }
 
-void UAnimSharingInstance::SetupSlaveComponent(uint8 CurrentState, uint32 ActorIndex)
+void UAnimSharingInstance::SetupFollowerComponent(uint8 CurrentState, uint32 ActorIndex)
 {
 	const FPerStateData& StateData = PerStateData[CurrentState];
 
 	if (StateData.Components.Num() == 0)
 	{	
-		UE_LOG(LogAnimationSharing, Warning, TEXT("No Master Components available for state %s, make sure to set up an Animation Sequence/Blueprint "), *StateEnum->GetDisplayNameTextByValue(CurrentState).ToString());
+		UE_LOG(LogAnimationSharing, Warning, TEXT("No Leader Components available for state %s, make sure to set up an Animation Sequence/Blueprint "), *StateEnum->GetDisplayNameTextByValue(CurrentState).ToString());
 		return;
 	}
 
 	if (!StateData.bIsOnDemand)
 	{
 		const uint32 PermutationIndex = DeterminePermutationIndex(ActorIndex, CurrentState);
-		SetPermutationSlaveComponent(CurrentState, ActorIndex, PermutationIndex);
+		SetPermutationFollowerComponent(CurrentState, ActorIndex, PermutationIndex);
 	}
 	else
 	{
@@ -1765,8 +1765,8 @@ void UAnimSharingInstance::SetupSlaveComponent(uint8 CurrentState, uint32 ActorI
 
 		if (OnDemandInstanceIndex != INDEX_NONE)
 		{
-			USkeletalMeshComponent* MasterComponent = StateData.Components[OnDemandInstances[OnDemandInstanceIndex].UsedPerStateComponentIndex];
-			SetMasterComponentForActor(ActorIndex, MasterComponent);
+			USkeletalMeshComponent* LeaderComponent = StateData.Components[OnDemandInstances[OnDemandInstanceIndex].UsedPerStateComponentIndex];
+			SetLeaderComponentForActor(ActorIndex, LeaderComponent);
 			OnDemandInstances[OnDemandInstanceIndex].ActorIndices.Add(ActorIndex);
 
 			PerActorData[ActorIndex].OnDemandInstanceIndex = OnDemandInstanceIndex;
@@ -1778,7 +1778,7 @@ void UAnimSharingInstance::SetupSlaveComponent(uint8 CurrentState, uint32 ActorI
 	}
 }
 
-void UAnimSharingInstance::SetPermutationSlaveComponent(uint8 StateIndex, uint32 ActorIndex, uint32 PermutationIndex)
+void UAnimSharingInstance::SetPermutationFollowerComponent(uint8 StateIndex, uint32 ActorIndex, uint32 PermutationIndex)
 {
 	const FPerStateData& StateData = PerStateData[StateIndex];
 
@@ -1791,7 +1791,7 @@ void UAnimSharingInstance::SetPermutationSlaveComponent(uint8 StateIndex, uint32
 	}
 #endif
 
-	SetMasterComponentForActor(ActorIndex, StateData.Components[PermutationIndex]);
+	SetLeaderComponentForActor(ActorIndex, StateData.Components[PermutationIndex]);
 	PerActorData[ActorIndex].PermutationIndex = PermutationIndex;
 	UAnimationSharingManager::SetDebugMaterial(StateData.Components[PermutationIndex], 1);
 }
@@ -1803,7 +1803,7 @@ uint32 UAnimSharingInstance::DeterminePermutationIndex(uint32 ActorIndex, uint8
 
 	// This can grow to be more intricate to take into account surrounding actors?
 	const uint32 PermutationIndex = FMath::RandHelper(Components.Num());
-	checkf(Components.IsValidIndex(PermutationIndex), TEXT("Not enough MasterComponents initialised!"));
+	checkf(Components.IsValidIndex(PermutationIndex), TEXT("Not enough LeaderComponents initialised!"));
 
 	return PermutationIndex;
 }
@@ -1909,7 +1909,7 @@ void UAnimSharingInstance::SwitchBetweenOnDemands(uint32 FromOnDemandInstanceInd
 	const uint32 ComponentIndex = Instance.UsedPerStateComponentIndex;
 	const uint32 StateIndex = Instance.State;
 	PerActorData[ActorIndex].PermutationIndex = 0;
-	SetMasterComponentForActor(ActorIndex, PerStateData[StateIndex].Components[ComponentIndex]);
+	SetLeaderComponentForActor(ActorIndex, PerStateData[StateIndex].Components[ComponentIndex]);
 }
 
 uint32 UAnimSharingInstance::SetupOnDemandInstance(uint8 StateIndex)
diff --git a/Engine/Plugins/Developer/AnimationSharing/Source/AnimationSharing/Public/AnimationSharingManager.h b/Engine/Plugins/Developer/AnimationSharing/Source/AnimationSharing/Public/AnimationSharingManager.h
--- a/Engine/Plugins/Developer/AnimationSharing/Source/AnimationSharing/Public/AnimationSharingManager.h
+++ b/Engine/Plugins/Developer/AnimationSharing/Source/AnimationSharing/Public/AnimationSharingManager.h
@@ -30,7 +30,7 @@ struct FBlendInstance
 {
 	/** Flag whether or not this instance is currently active */
 	bool bActive;
-	/** Flag whether or not the actor's part of this have been setup as a slave component to the blend actor, this is done so the blend actor atleast ticks once (otherwise it can pop from the previous blend end pose) */
+	/** Flag whether or not the actor's part of this have been setup as a follower component to the blend actor, this is done so the blend actor atleast ticks once (otherwise it can pop from the previous blend end pose) */
 	bool bBlendStarted;
 	/** Flag whether or not this instance is blending towards an on-demand state */
 	bool bOnDemand;
@@ -43,13 +43,13 @@ struct FBlendInstance
 	/** State value to blend to */
 	uint8 StateTo;
 
-	/** Permutation indices from and to which we are blending, used to ensure we 'forward' the actor to the correct MasterPoseComponent when finished blending */
+	/** Permutation indices from and to which we are blending, used to ensure we 'forward' the actor to the correct LeaderPoseComponent when finished blending */
 	uint32 FromPermutationIndex;
 	uint32 ToPermutationIndex;
 
 	/** Actor used for blending between the two states */
 	FTransitionBlendInstance* TransitionBlendInstance;
-	/** Indices of actors who are set up as slaves to BlendActor's main skeletal mesh component */
+	/** Indices of actors who are set up as followers to BlendActor's main skeletal mesh component */
 	TArray<uint32> ActorIndices;
 
 	/** Optional index into OnDemandInstances from which we are blending */
@@ -89,7 +89,7 @@ struct FOnDemandInstance
 	/** Permutation index that we are blending to before the end of the animation */
 	uint32 BlendToPermutationIndex;
 
-	/** Indices of actors who are set up as slaves to the skeletal mesh component running the animation */
+	/** Indices of actors who are set up as followers to the skeletal mesh component running the animation */
 	TArray<uint32> ActorIndices;
 };
 
@@ -212,8 +212,13 @@ struct FPerStateData
 	TBitArray<> InUseComponentFrameBits;
 	TBitArray<> PreviousInUseComponentFrameBits;
 
-	/** Bits keeping track whether or not any of the slave components requires the master component to tick */
+	/** Bits keeping track whether or not any of the follower components requires the leader component to tick */
+	TBitArray<> FollowerTickRequiredFrameBits;
+
+#if UE_EDITOR
+	UE_DEPRECATED(5.1, "This member has been deprecated. Please use FollowerTickRequiredFrameBits")
 	TBitArray<> SlaveTickRequiredFrameBits;
+#endif // UE_EDITOR
 
 	/** Length of the animations used for an on-demand state, array as it could contain different animation permutations */
 	TArray<float> AnimationLengths;
@@ -250,7 +255,7 @@ struct FPerActorData
 	/** Cached significance value */
 	float SignificanceValue;
 
-	/** Flag whether or not this actor requires the master component to tick */
+	/** Flag whether or not this actor requires the leader component to tick */
 	bool bRequiresTick;
 
 	/** Index to blend instance which is currently driving this actor's animation */
@@ -337,7 +342,7 @@ public:
 		return bShouldTick;
 	}
 
-	/** Ticks all currently running additive animation instances, this checks whether or not it has finished yet and sets the base-component as the master component when it has */
+	/** Ticks all currently running additive animation instances, this checks whether or not it has finished yet and sets the base-component as the leader component when it has */
 	void TickAdditiveInstances();
 
 	/** Ticks all Actor Data entries and determines their current state, if changed since last tick it will alter their animation accordingly */
@@ -346,7 +351,7 @@ public:
 	/** Ticks various types of debugging data / drawing (not active in shipping build) */
 	void TickDebugInformation();
 
-	/** Ticks all unique animation states, this checks which components are currently used and turns of those which currently don't have any slaves */
+	/** Ticks all unique animation states, this checks which components are currently used and turns of those which currently don't have any followerss */
 	void TickAnimationStates();
 
 	/** Removal functions which also make sure that indices of other data-structres are correctly remapped */
@@ -355,7 +360,7 @@ public:
 	void RemoveAdditiveInstance(int32 InstanceIndex);
 	void RemoveComponent(int32 ComponentIndex);
 
-	/** Removal functions which also make sure the actor is set to the correct master pose component */
+	/** Removal functions which also make sure the actor is set to the correct leader pose component */
 	void RemoveFromCurrentBlend(int32 ActorIndex);
 	void RemoveFromCurrentOnDemand(int32 ActorIndex);
 
@@ -365,14 +370,23 @@ public:
 	/**  Frees up an Additive Animation instance and resets it state*/
 	void FreeAdditiveInstance(uint8 StateIndex, FAdditiveAnimationInstance* Instance);
 
-	/** Sets up all components of an actor to be slaves of Component */
-	void SetMasterComponentForActor(uint32 ActorIndex, USkeletalMeshComponent* Component);
+	/** Sets up all components of an actor to be followers of Component */
+	void SetLeaderComponentForActor(uint32 ActorIndex, USkeletalMeshComponent* Component);
+
+	UE_DEPRECATED(5.1, "This function has been deprecated. Please use SetLeaderComponentForActor instead.")
+	void SetMasterComponentForActor(uint32 ActorIndex, USkeletalMeshComponent* Component) { SetLeaderComponentForActor(ActorIndex, Component); }
+
+	/** Sets up the correct LeaderPoseComponent for the passed in Component and State indices */
+	void SetupFollowerComponent(uint8 CurrentState, uint32 ActorIndex);
 
-	/** Sets up the correct MasterPoseComponent for the passed in Component and State indices */
-	void SetupSlaveComponent(uint8 CurrentState, uint32 ActorIndex);
+	UE_DEPRECATED(5.1, "This function has been deprecated. Please use SetupFollowerComponent instead.")
+	void SetupSlaveComponent(uint8 CurrentState, uint32 ActorIndex) { SetupFollowerComponent(CurrentState, ActorIndex); }
 
-	/** Sets up the correct MasterPoseComponent according to the state and permutation indices */
-	void SetPermutationSlaveComponent(uint8 StateIndex, uint32 ActorIndex, uint32 PermutationIndex);
+	/** Sets up the correct LeaderPoseComponent according to the state and permutation indices */
+	void SetPermutationFollowerComponent(uint8 StateIndex, uint32 ActorIndex, uint32 PermutationIndex);
+
+	UE_DEPRECATED(5.1, "This function has been deprecated. Please use SetPermutationFollowerComponent instead.")
+	void SetPermutationSlaveComponent(uint8 StateIndex, uint32 ActorIndex, uint32 PermutationIndex) { SetPermutationFollowerComponent(StateIndex, ActorIndex, PermutationIndex); }
 	
 	/** Determines a permutation index for the given actor and state */
 	uint32 DeterminePermutationIndex(uint32 ActorIndex, uint8 State) const;
@@ -380,7 +394,7 @@ public:
 	/** Marks the component as either used/not-used, this is used to disable ticking of components which are not in use*/
 	void SetComponentUsage(bool bUsage, uint8 StateIndex, uint32 ComponentIndex);
 
-	/** Sets the whether or not any of the slave components are visible */
+	/** Sets the whether or not any of the follower components are visible */
 	void SetComponentTick(uint8 StateIndex, uint32 ComponentIndex);
 
 	/** Actors currently registered to be animation driven by the AnimManager using this setup */
@@ -514,8 +528,11 @@ public:
 	/** Ensures all currently registered actors are removed */
 	void UnregisterAllActors();	
 
-	/** Sets the visibility of currently used Master Pose Components */
-	void SetMasterComponentsVisibility(bool bVisible);
+	/** Sets the visibility of currently used Leader Pose Components */
+	void SetLeaderComponentsVisibility(bool bVisible);
+	
+	UE_DEPRECATED(5.1, "This member has been deprecated. Please use SetLeaderComponentsVisibility")
+	void SetMasterComponentsVisibility(bool bVisible) { SetLeaderComponentsVisibility(bVisible); }
 
 	/** Initialize sharing data structures */
 	void Initialise(const UAnimationSharingSetup* InSetup);
diff --git a/Engine/Plugins/Developer/AnimationSharing/Source/AnimationSharing/Public/AnimationSharingTypes.h b/Engine/Plugins/Developer/AnimationSharing/Source/AnimationSharing/Public/AnimationSharingTypes.h
--- a/Engine/Plugins/Developer/AnimationSharing/Source/AnimationSharing/Public/AnimationSharingTypes.h
+++ b/Engine/Plugins/Developer/AnimationSharing/Source/AnimationSharing/Public/AnimationSharingTypes.h
@@ -89,7 +89,7 @@ public:
 	UPROPERTY(EditAnywhere, Category = AnimationSharing, meta = (EditCondition = "bOnDemand", UIMin="0.0", UIMax="1.0", ClampMin = "0.0", ClampMax = "1.0"))
 	float WiggleTimePercentage;
 	
-	/** Whether or not this animation requires curves or morphtargets to function correctly for slave components */
+	/** Whether or not this animation requires curves or morphtargets to function correctly for follower components */
 	UPROPERTY(EditAnywhere, AdvancedDisplay, Category = AnimationSharing)
 	bool bRequiresCurves;
 };
@@ -179,7 +179,7 @@ public:
 	UPROPERTY(EditAnywhere, Category = AnimationSharing, meta = (ClampMin = "1", UIMin = "1"))
 	FPerPlatformInt MaximumNumberConcurrentBlends;
 
-	/** Significance value tied to whether or not the master pose components should be ticking */
+	/** Significance value tied to whether or not the leader pose components should be ticking */
 	UPROPERTY(EditAnywhere, Category = AnimationSharing, meta = (ClampMin = "0.0", UIMin = "0.0"))
 	FPerPlatformFloat TickSignificanceValue;
 };
