* Added UE_SCOPED_ENGINE_ACTIVITY in a few strategic places to more easily identify where time is spent when starting the runtime

diff --git a/Engine/Plugins/Developer/AnimationSharing/Source/AnimationSharing/Private/AnimationSharingModule.cpp b/Engine/Plugins/Developer/AnimationSharing/Source/AnimationSharing/Private/AnimationSharingModule.cpp
--- a/Engine/Plugins/Developer/AnimationSharing/Source/AnimationSharing/Private/AnimationSharingModule.cpp
+++ b/Engine/Plugins/Developer/AnimationSharing/Source/AnimationSharing/Private/AnimationSharingModule.cpp
@@ -35,6 +35,7 @@ bool FAnimSharingModule::CreateAnimationSharingManager(UWorld* InWorld, const UA
 {
 	if (InWorld && InWorld->IsGameWorld() && Setup && UAnimationSharingManager::AnimationSharingEnabled() && !WorldAnimSharingManagers.Contains(InWorld))
 	{
+		UE_SCOPED_ENGINE_ACTIVITY(TEXT("Initializing Animation Sharing"));
 		UAnimationSharingManager* Manager = NewObject<UAnimationSharingManager>(InWorld);
 		Manager->Initialise(Setup);
 		WorldAnimSharingManagers.Add(InWorld, Manager);
diff --git a/Engine/Plugins/Experimental/GameFeatures/Source/GameFeatures/Private/GameFeatureAction_AddComponents.cpp b/Engine/Plugins/Experimental/GameFeatures/Source/GameFeatures/Private/GameFeatureAction_AddComponents.cpp
--- a/Engine/Plugins/Experimental/GameFeatures/Source/GameFeatures/Private/GameFeatureAction_AddComponents.cpp
+++ b/Engine/Plugins/Experimental/GameFeatures/Source/GameFeatures/Private/GameFeatureAction_AddComponents.cpp
@@ -111,6 +111,7 @@ void UGameFeatureAction_AddComponents::AddToWorld(const FWorldContext& WorldCont
 				{
 					if (!Entry.ActorClass.IsNull())
 					{
+						UE_SCOPED_ENGINE_ACTIVITY(TEXT("Adding component to world %s (%s)"), *World->GetDebugDisplayName(), *Entry.ComponentClass.ToString());
 						TSubclassOf<UActorComponent> ComponentClass = Entry.ComponentClass.LoadSynchronous();
 						if (ComponentClass)
 						{
diff --git a/Engine/Plugins/Runtime/GameplayAbilities/Source/GameplayAbilities/Private/GameplayCueManager.cpp b/Engine/Plugins/Runtime/GameplayAbilities/Source/GameplayAbilities/Private/GameplayCueManager.cpp
--- a/Engine/Plugins/Runtime/GameplayAbilities/Source/GameplayAbilities/Private/GameplayCueManager.cpp
+++ b/Engine/Plugins/Runtime/GameplayAbilities/Source/GameplayAbilities/Private/GameplayCueManager.cpp
@@ -601,6 +601,8 @@ bool UGameplayCueManager::ShouldAsyncLoadRuntimeObjectLibraries() const
 
 void UGameplayCueManager::InitializeRuntimeObjectLibrary()
 {
+	UE_SCOPED_ENGINE_ACTIVITY(TEXT("Initializing GameplayCueManager Runtime Object Library"));
+
 	RuntimeGameplayCueObjectLibrary.Paths = GetAlwaysLoadedGameplayCuePaths();
 	if (RuntimeGameplayCueObjectLibrary.CueSet == nullptr)
 	{
diff --git a/Engine/Source/Editor/UnrealEd/Private/FileHelpers.cpp b/Engine/Source/Editor/UnrealEd/Private/FileHelpers.cpp
--- a/Engine/Source/Editor/UnrealEd/Private/FileHelpers.cpp
+++ b/Engine/Source/Editor/UnrealEd/Private/FileHelpers.cpp
@@ -2756,6 +2756,8 @@ static void NotifyBSPNeedsRebuild(const FString& PackageName)
  */
 bool FEditorFileUtils::LoadMap(const FString& InFilename, bool LoadAsTemplate, bool bShowProgress)
 {
+	UE_SCOPED_ENGINE_ACTIVITY(TEXT("Loading Map %s"), *InFilename);
+
 	double LoadStartTime = FStudioAnalytics::GetAnalyticSeconds();
 	
 	if (GEditor->WarnIfLightingBuildIsCurrentlyRunning())
diff --git a/Engine/Source/Editor/UnrealEd/Private/UnrealEdEngine.cpp b/Engine/Source/Editor/UnrealEd/Private/UnrealEdEngine.cpp
--- a/Engine/Source/Editor/UnrealEd/Private/UnrealEdEngine.cpp
+++ b/Engine/Source/Editor/UnrealEd/Private/UnrealEdEngine.cpp
@@ -172,6 +172,8 @@ void UUnrealEdEngine::Init(IEngineLoop* InEngineLoop)
 
 	if (FPaths::IsProjectFilePathSet() && GIsEditor && !FApp::IsUnattended())
 	{
+
+		UE_SCOPED_ENGINE_ACTIVITY(TEXT("Initializing AutoRepimportManager"));
 		AutoReimportManager = NewObject<UAutoReimportManager>();
 		AutoReimportManager->Initialize();
 	}
diff --git a/Engine/Source/Runtime/Engine/Private/TextureCompiler.cpp b/Engine/Source/Runtime/Engine/Private/TextureCompiler.cpp
--- a/Engine/Source/Runtime/Engine/Private/TextureCompiler.cpp
+++ b/Engine/Source/Runtime/Engine/Private/TextureCompiler.cpp
@@ -449,6 +449,8 @@ void FTextureCompilingManager::PostCompilation(TArrayView<UTexture* const> InCom
 
 void FTextureCompilingManager::FinishAllCompilation()
 {
+	UE_SCOPED_ENGINE_ACTIVITY(TEXT("Finish Texture Compilation"));
+
 	check(IsInGameThread());
 	TRACE_CPUPROFILER_EVENT_SCOPE(FTextureCompilingManager::FinishAllCompilation)
 
