UAnimationSharingManager::RegisterActorWithSkeleton incorrectly calls SetupSlaveComponent multiple times

diff --git a/Engine/Plugins/Developer/AnimationSharing/Source/AnimationSharing/Private/AnimationSharingManager.cpp b/Engine/Plugins/Developer/AnimationSharing/Source/AnimationSharing/Private/AnimationSharingManager.cpp
--- a/Engine/Plugins/Developer/AnimationSharing/Source/AnimationSharing/Private/AnimationSharingManager.cpp
+++ b/Engine/Plugins/Developer/AnimationSharing/Source/AnimationSharing/Private/AnimationSharingManager.cpp
@@ -434,13 +434,12 @@ void UAnimationSharingManager::RegisterActorWithSkeleton(AActor* InActor, const
 				Component->SetComponentTickEnabled(false);
 				Component->PrimaryComponentTick.bCanEverTick = false;
 				Component->bIgnoreLeaderPoseComponentLOD = true;
-
+				
 				ActorData.ComponentIndices.Add(Data->PerComponentData.Num() - 1);
-
-				const int32 ComponentIndex = Data->PerComponentData.Num() - 1;
-				Data->SetupFollowerComponent(ActorData.CurrentState, ActorIndex);
 			}
 
+			Data->SetupFollowerComponent(ActorData.CurrentState, ActorIndex);
+
 			if (Data->PerStateData[ActorData.CurrentState].bIsOnDemand && ActorData.OnDemandInstanceIndex != INDEX_NONE)
 			{
 				// We will have setup an on-demand instance so we need to kick it off here before we next tick
