Truncation fixes for AnimationSharingEd

diff --git a/Engine/Plugins/Developer/AnimationSharing/Source/AnimationSharingEd/AnimationSharingEd.Build.cs b/Engine/Plugins/Developer/AnimationSharing/Source/AnimationSharingEd/AnimationSharingEd.Build.cs
--- a/Engine/Plugins/Developer/AnimationSharing/Source/AnimationSharingEd/AnimationSharingEd.Build.cs
+++ b/Engine/Plugins/Developer/AnimationSharing/Source/AnimationSharingEd/AnimationSharingEd.Build.cs
@@ -22,6 +22,7 @@ public class AnimationSharingEd : ModuleRules
                 "PropertyEditor",
             }
 		);
-    }
 
+		UnsafeTypeCastWarningLevel = WarningLevel.Error;
+	}
 }
diff --git a/Engine/Plugins/Developer/AnimationSharing/Source/AnimationSharingEd/Private/SetupDetailsViewCustomizations.cpp b/Engine/Plugins/Developer/AnimationSharing/Source/AnimationSharingEd/Private/SetupDetailsViewCustomizations.cpp
--- a/Engine/Plugins/Developer/AnimationSharing/Source/AnimationSharingEd/Private/SetupDetailsViewCustomizations.cpp
+++ b/Engine/Plugins/Developer/AnimationSharing/Source/AnimationSharingEd/Private/SetupDetailsViewCustomizations.cpp
@@ -293,7 +293,7 @@ void FAnimationStateEntryCustomization::SelectedEnumChanged(TSharedPtr<FString>
 	{
 		if (Selection && SelectInfo != ESelectInfo::Type::Direct)
 		{
-			uint8 NewEnumValue = EnumClass->GetValueByIndex(ComboBoxItems.IndexOfByKey(Selection));
+			const uint8 NewEnumValue = IntCastChecked<uint8>(EnumClass->GetValueByIndex(ComboBoxItems.IndexOfByKey(Selection)));
 			PropertyHandle->SetValue(NewEnumValue);
 		}
 
