Truncation fixes for AnimationSharing

diff --git a/Engine/Plugins/Developer/AnimationSharing/Source/AnimationSharing/AnimationSharing.Build.cs b/Engine/Plugins/Developer/AnimationSharing/Source/AnimationSharing/AnimationSharing.Build.cs
--- a/Engine/Plugins/Developer/AnimationSharing/Source/AnimationSharing/AnimationSharing.Build.cs
+++ b/Engine/Plugins/Developer/AnimationSharing/Source/AnimationSharing/AnimationSharing.Build.cs
@@ -21,6 +21,8 @@ public class AnimationSharing : ModuleRules
         if (Target.Type == TargetType.Editor)
         {
             PrivateDependencyModuleNames.Add("TargetPlatform");
-        }     
-    }
+        }
+
+		UnsafeTypeCastWarningLevel = WarningLevel.Error;
+	}
 }
diff --git a/Engine/Plugins/Developer/AnimationSharing/Source/AnimationSharing/Private/AnimationSharingManager.cpp b/Engine/Plugins/Developer/AnimationSharing/Source/AnimationSharing/Private/AnimationSharingManager.cpp
--- a/Engine/Plugins/Developer/AnimationSharing/Source/AnimationSharing/Private/AnimationSharingManager.cpp
+++ b/Engine/Plugins/Developer/AnimationSharing/Source/AnimationSharing/Private/AnimationSharingManager.cpp
@@ -260,7 +260,7 @@ void UAnimationSharingManager::Tick(float DeltaTime)
 {
 	SCOPE_CYCLE_COUNTER(STAT_AnimationSharing_Tick);
 	
-	const float WorldTime = GetWorld()->GetTimeSeconds();
+	const float WorldTime = static_cast<float>(GetWorld()->GetTimeSeconds());
 
 	/** Keeping track of currently running instances / animations for debugging purposes */
 	int32 TotalNumBlends = 0;
@@ -440,7 +440,7 @@ void UAnimationSharingManager::RegisterActorWithSkeleton(AActor* InActor, const
 				Data->OnDemandInstances[ActorData.OnDemandInstanceIndex].StartTime = Data->WorldTime;
 			}
 
-			const int32 ActorHandle = CreateActorHandle(Handle, ActorIndex);
+			const int32 ActorHandle = CreateActorHandle(IntCastChecked<uint8>(Handle), ActorIndex);
 			ActorData.UpdateActorHandleDelegate.ExecuteIfBound(ActorHandle);
 		}
 	}
@@ -531,7 +531,7 @@ void UAnimationSharingManager::UnregisterActor(AActor* InActor)
 				}
 
 				// Make sure we update the handle on the swapped actor
-				SkeletonData->PerActorData[SwapIndex].UpdateActorHandleDelegate.ExecuteIfBound(CreateActorHandle(SkeletonIndex, ActorIndex));
+				SkeletonData->PerActorData[SwapIndex].UpdateActorHandleDelegate.ExecuteIfBound(CreateActorHandle(IntCastChecked<uint8>(SkeletonIndex), ActorIndex));
 			}			
 
 			SkeletonData->PerActorData.RemoveAtSwap(ActorIndex, 1, false);
@@ -720,7 +720,7 @@ uint8 UAnimSharingInstance::DetermineStateForActor(uint32 ActorIndex, bool& bSho
 		StateProcessor->ProcessActorState(State, RegisteredActors[ActorIndex], ActorData.CurrentState, ActorData.OnDemandInstanceIndex != INDEX_NONE ? OnDemandInstances[ActorData.OnDemandInstanceIndex].State : INDEX_NONE, bShouldProcess);
 	}
 	
-	return FMath::Max(0, State);
+	return FMath::Max((uint8)0, IntCastChecked<uint8>(State));
 }
 
 bool UAnimSharingInstance::Setup(UAnimationSharingManager* AnimationSharingManager, const FPerSkeletonAnimationSharingSetup& SkeletonSetup, const FAnimationSharingScalability* InScalabilitySettings, uint32 Index)
@@ -786,7 +786,7 @@ bool UAnimSharingInstance::Setup(UAnimationSharingManager* AnimationSharingManag
 		if (!bErrors && ScalabilitySettings->UseBlendTransitions.Default)
 		{
 			const uint32 TotalNumberOfBlendActorsRequired = ScalabilitySettings->MaximumNumberConcurrentBlends.Default;
-			const float ZOffset = Index * SkeletalMeshBounds.Z * 2.f;
+			const float ZOffset = static_cast<float>((double)Index * SkeletalMeshBounds.Z * 2.0);
 			for (uint32 BlendIndex = 0; BlendIndex < TotalNumberOfBlendActorsRequired; ++BlendIndex)
 			{
 				const FVector SpawnLocation(BlendIndex * SkeletalMeshBounds.X, 0.f, ZOffset + SkeletalMeshBounds.Z);
@@ -817,7 +817,7 @@ bool UAnimSharingInstance::Setup(UAnimationSharingManager* AnimationSharingManag
 void UAnimSharingInstance::SetupState(FPerStateData& StateData, const FAnimationStateEntry& StateEntry, USkeletalMesh* SkeletalMesh, const FPerSkeletonAnimationSharingSetup& SkeletonSetup, uint32 Index)
 {
 	/** Used for placing components into rows / columns at origin for debugging purposes */
-	const float ZOffset = Index * SkeletalMeshBounds.Z * 2.f;
+	const float ZOffset = static_cast<float>((double)Index * SkeletalMeshBounds.Z * 2.0);
 
 	/** Setup overall data and flags */
 	StateData.bIsOnDemand = StateEntry.bOnDemand;
@@ -911,8 +911,8 @@ void UAnimSharingInstance::SetupState(FPerStateData& StateData, const FAnimation
 							AnimInstance->AnimationToPlay = AnimSequence;
 							if (InstanceIndex > 0)
 							{
-								const float Steps = (AnimSequence->GetPlayLength() * 0.9f) / (NumInstances);
-								const float StartTimeOffset = Steps * InstanceIndex;
+								const float Steps = (AnimSequence->GetPlayLength() * 0.9f) / static_cast<float>(NumInstances);
+								const float StartTimeOffset = Steps * static_cast<float>(InstanceIndex);
 								AnimInstance->PermutationTimeOffset = StartTimeOffset;
 							}
 
@@ -920,7 +920,7 @@ void UAnimSharingInstance::SetupState(FPerStateData& StateData, const FAnimation
 
 							AnimInstance->Instance = this;
 							AnimInstance->StateIndex = StateEntry.State;
-							AnimInstance->ComponentIndex = Components.Num();
+							AnimInstance->ComponentIndex = IntCastChecked<uint8>(Components.Num());
 
 							/** Set the current animation length length */
 							StateData.AnimationLengths.Add(AnimSequence->GetPlayLength());
@@ -939,8 +939,8 @@ void UAnimSharingInstance::SetupState(FPerStateData& StateData, const FAnimation
 						{
 							if (InstanceIndex > 0)
 							{
-								const float Steps = (AnimSequence->GetPlayLength() * 0.9f) / (NumInstances);
-								const float StartTimeOffset = Steps * InstanceIndex;
+								const float Steps = (AnimSequence->GetPlayLength() * 0.9f) / static_cast<float>(NumInstances);
+								const float StartTimeOffset = Steps * static_cast<float>(InstanceIndex);
 								Component->SetPosition(StartTimeOffset, false);
 							}
 						}
@@ -1044,7 +1044,7 @@ void UAnimSharingInstance::TickDebugInformation()
 				const uint32 BlendInstanceIndex = ActorData.BlendInstanceIndex;
 				if (BlendInstanceIndex != INDEX_NONE && BlendInstances.IsValidIndex(BlendInstanceIndex))
 				{
-					const float TimeLeft = BlendInstances[BlendInstanceIndex].BlendTime - (GetWorld()->GetTimeSeconds() - BlendInstances[BlendInstanceIndex].EndTime);
+					const float TimeLeft = BlendInstances[BlendInstanceIndex].BlendTime - static_cast<float>(GetWorld()->GetTimeSeconds() - BlendInstances[BlendInstanceIndex].EndTime);
 					return FString::Printf(TEXT("Blending states - %s to %s [%1.3f] (%i)"), *StateEnum->GetDisplayNameTextByValue(BlendInstances[BlendInstanceIndex].StateFrom).ToString(), *StateEnum->GetDisplayNameTextByValue(BlendInstances[BlendInstanceIndex].StateTo).ToString(), TimeLeft, ActorData.BlendInstanceIndex);
 				}
 
@@ -1164,7 +1164,7 @@ void UAnimSharingInstance::TickOnDemandInstances()
 			// Set the components to their current state animation
 			for (uint32 ActorIndex : Instance.ActorIndices)
 			{					
-				const uint32 CurrentState = PerActorData[ActorIndex].CurrentState;
+				const uint8 CurrentState = PerActorData[ActorIndex].CurrentState;
 				// Return to the previous active animation state 
 				if (Instance.bReturnToPreviousState)
 				{
@@ -1294,7 +1294,7 @@ void UAnimSharingInstance::TickAdditiveInstances()
 
 		if (Instance.bActive)
 		{
-			const float WorldTimeSeconds = GetWorld()->GetTimeSeconds();
+			const float WorldTimeSeconds = static_cast<float>(GetWorld()->GetTimeSeconds());
 			if (WorldTimeSeconds >= Instance.EndTime)
 			{
 				// Finish
@@ -1413,7 +1413,7 @@ void UAnimSharingInstance::TickActorStates()
 							if (ActorData.bRunningOnDemand)
 							{
 								/** Setup a blend between the current and a new instance*/
-								const uint32 BlendInstanceIndex = SetupBlendBetweenOnDemands(ActorData.OnDemandInstanceIndex, OnDemandIndex, ActorIndex);
+								const uint32 BlendInstanceIndex = SetupBlendBetweenOnDemands(IntCastChecked<uint8>(ActorData.OnDemandInstanceIndex), OnDemandIndex, ActorIndex);
 								ActorData.BlendInstanceIndex = BlendInstanceIndex;
 							}
 							else
@@ -1784,7 +1784,7 @@ void UAnimSharingInstance::SetPermutationFollowerComponent(uint8 StateIndex, uin
 #endif
 
 	SetLeaderComponentForActor(ActorIndex, StateData.Components[PermutationIndex]);
-	PerActorData[ActorIndex].PermutationIndex = PermutationIndex;
+	PerActorData[ActorIndex].PermutationIndex = IntCastChecked<uint8>(PermutationIndex);
 	UAnimationSharingManager::SetDebugMaterial(StateData.Components[PermutationIndex], 1);
 }
 
@@ -1829,7 +1829,7 @@ uint32 UAnimSharingInstance::SetupBlend(uint8 FromState, uint8 ToState, uint32 A
 			BlendInstance->StateTo = ToState;
 			BlendInstance->BlendTime = CalculateBlendTime( ToState);
 			BlendInstance->bOnDemand = bOnDemand;
-			BlendInstance->EndTime = GetWorld()->GetTimeSeconds() + BlendInstance->BlendTime;
+			BlendInstance->EndTime = static_cast<float>(GetWorld()->GetTimeSeconds()) + BlendInstance->BlendTime;
 			BlendInstance->TransitionBlendInstance = BlendInstanceStack.GetInstance();
 
 			BlendInstance->TransitionBlendInstance->GetComponent()->SetComponentTickEnabled(true);
@@ -1935,7 +1935,7 @@ uint32 UAnimSharingInstance::SetupOnDemandInstance(uint8 StateIndex)
 				Instance.StartTime = 0.f;
 				Instance.BlendToPermutationIndex = INDEX_NONE;
 
-				const float WorldTimeSeconds = GetWorld()->GetTimeSeconds();
+				const float WorldTimeSeconds = static_cast<float>(GetWorld()->GetTimeSeconds());
 				Instance.EndTime = WorldTimeSeconds + StateData.AnimationLengths[AvailableIndex];
 				Instance.StartBlendTime = Instance.EndTime - CalculateBlendTime(StateIndex);
 
@@ -2011,7 +2011,7 @@ uint32 UAnimSharingInstance::SetupAdditiveInstance(uint8 StateIndex, uint8 FromS
 		Instance.bActive = false;
 		Instance.AdditiveAnimationInstance = AnimationInstance;
 		Instance.BaseComponent = PerStateData[FromState].Components[StateComponentIndex];
-		const float WorldTimeSeconds = GetWorld()->GetTimeSeconds();
+		const float WorldTimeSeconds = static_cast<float>(GetWorld()->GetTimeSeconds());
 		Instance.EndTime = WorldTimeSeconds + StateData.AdditiveAnimationSequence->GetPlayLength();
 		Instance.State = StateIndex;
 		Instance.UsedPerStateComponentIndex = PerStateData[StateIndex].Components.IndexOfByKey(AnimationInstance->GetComponent());
@@ -2071,7 +2071,7 @@ void UAnimSharingInstance::KickoffInstances()
 
 			for (uint32 ActorIndex : BlendInstance.ActorIndices)
 			{
-				PerActorData[ActorIndex].PermutationIndex = BlendInstance.ToPermutationIndex;
+				PerActorData[ActorIndex].PermutationIndex = IntCastChecked<uint8>(BlendInstance.ToPermutationIndex);
 				PerActorData[ActorIndex].bBlending = true;
 			}
 
