Added option in Anim Sharing Setup to decide whether or not the skeletal mesh components created for each state can use material parameter caching

diff --git a/Engine/Plugins/Developer/AnimationSharing/Source/AnimationSharing/Private/AnimationSharingManager.cpp b/Engine/Plugins/Developer/AnimationSharing/Source/AnimationSharing/Private/AnimationSharingManager.cpp
--- a/Engine/Plugins/Developer/AnimationSharing/Source/AnimationSharing/Private/AnimationSharingManager.cpp
+++ b/Engine/Plugins/Developer/AnimationSharing/Source/AnimationSharing/Private/AnimationSharingManager.cpp
@@ -796,6 +796,7 @@ bool UAnimSharingInstance::Setup(UAnimationSharingManager* AnimationSharingManag
 				BlendComponent->SetRelativeLocation(SpawnLocation);
 				BlendComponent->SetSkeletalMesh(SkeletalMesh);
 				BlendComponent->SetVisibility(GLeaderComponentsVisible == 1);
+				BlendComponent->bEnableMaterialParameterCaching = SkeletonSetup.bEnableMaterialParameterCaching;
 
 				BlendComponent->PrimaryComponentTick.AddPrerequisite(AnimSharingManager, AnimSharingManager->GetTickFunction());
 
@@ -902,6 +903,7 @@ void UAnimSharingInstance::SetupState(FPerStateData& StateData, const FAnimation
 					Component->SetForcedLOD(1);
 					Component->SetVisibility(GLeaderComponentsVisible == 1);
 					Component->bPropagateCurvesToFollowers = StateEntry.bRequiresCurves;
+					Component->bEnableMaterialParameterCaching = SkeletonSetup.bEnableMaterialParameterCaching;
 
 					if (AnimBPClass != nullptr && AnimSequence != nullptr && !StateEntry.bOnDemand)
 					{
@@ -964,6 +966,7 @@ void UAnimSharingInstance::SetupState(FPerStateData& StateData, const FAnimation
 					AdditiveComponent->SetRelativeLocation(SpawnLocation);
 					AdditiveComponent->SetSkeletalMesh(SkeletalMesh);
 					AdditiveComponent->SetVisibility(GLeaderComponentsVisible == 1);
+					AdditiveComponent->bEnableMaterialParameterCaching = SkeletonSetup.bEnableMaterialParameterCaching;
 
 					AdditiveComponent->PrimaryComponentTick.AddPrerequisite(AnimSharingManager, AnimSharingManager->GetTickFunction());
 
diff --git a/Engine/Plugins/Developer/AnimationSharing/Source/AnimationSharing/Public/AnimationSharingTypes.h b/Engine/Plugins/Developer/AnimationSharing/Source/AnimationSharing/Public/AnimationSharingTypes.h
--- a/Engine/Plugins/Developer/AnimationSharing/Source/AnimationSharing/Public/AnimationSharingTypes.h
+++ b/Engine/Plugins/Developer/AnimationSharing/Source/AnimationSharing/Public/AnimationSharingTypes.h
@@ -129,7 +129,7 @@ struct FPerSkeletonAnimationSharingSetup
 {
 	GENERATED_BODY()
 public:
-	FPerSkeletonAnimationSharingSetup() : Skeleton(nullptr), SkeletalMesh(nullptr), BlendAnimBlueprint(nullptr), AdditiveAnimBlueprint(nullptr), StateProcessorClass(nullptr) {}
+	FPerSkeletonAnimationSharingSetup() : Skeleton(nullptr), SkeletalMesh(nullptr), BlendAnimBlueprint(nullptr), AdditiveAnimBlueprint(nullptr), StateProcessorClass(nullptr), bEnableMaterialParameterCaching(false){}
 
 	/** Skeleton compatible with the animation sharing setup */
 	UPROPERTY(EditAnywhere, Category = AnimationSharing)
@@ -151,6 +151,10 @@ public:
 	UPROPERTY(EditAnywhere, Category = AnimationSharing)
 	TSubclassOf<UAnimationSharingStateProcessor> StateProcessorClass;
 
+	/** Whether or not to enable material parameter caching on the skeletal mesh components created for each state */
+	UPROPERTY(EditAnywhere, Category = AnimationSharing)
+	bool bEnableMaterialParameterCaching;
+
 	/** Definition of different animation states */
 	UPROPERTY(EditAnywhere, Category = AnimationSharing)
 	TArray<FAnimationStateEntry> AnimationStates;
