Switch animation sharing to hold SkeletalMeshComponents as WeakObjectPtr to deal with GC/stale components (preventive)

diff --git a/Engine/Plugins/Developer/AnimationSharing/Source/AnimationSharing/Private/AnimationSharingManager.cpp b/Engine/Plugins/Developer/AnimationSharing/Source/AnimationSharing/Private/AnimationSharingManager.cpp
--- a/Engine/Plugins/Developer/AnimationSharing/Source/AnimationSharing/Private/AnimationSharingManager.cpp
+++ b/Engine/Plugins/Developer/AnimationSharing/Source/AnimationSharing/Private/AnimationSharingManager.cpp
@@ -422,7 +422,7 @@ void UAnimationSharingManager::RegisterActorWithSkeleton(AActor* InActor, const
 			{
 				FPerComponentData& ComponentData = Data->PerComponentData.AddZeroed_GetRef();
 				ComponentData.ActorIndex = ActorIndex;
-				ComponentData.Component = Component;
+				ComponentData.WeakComponent = Component;
 
 				Component->SetComponentTickEnabled(false);
 				Component->PrimaryComponentTick.bCanEverTick = false;
@@ -475,9 +475,12 @@ void UAnimationSharingManager::UnregisterActor(AActor* InActor)
 
 			for (int32 ComponentIndex : ActorData.ComponentIndices)
 			{
-				SkeletonData->PerComponentData[ComponentIndex].Component->SetLeaderPoseComponent(nullptr, true);
-				SkeletonData->PerComponentData[ComponentIndex].Component->PrimaryComponentTick.bCanEverTick = true;
-				SkeletonData->PerComponentData[ComponentIndex].Component->SetComponentTickEnabled(true);
+				if (USkeletalMeshComponent* Component = SkeletonData->GetComponent(SkeletonData->PerComponentData[ComponentIndex]))
+				{
+					Component->SetLeaderPoseComponent(nullptr, true);
+					Component->PrimaryComponentTick.bCanEverTick = true;
+					Component->SetComponentTickEnabled(true);
+				}
 				SkeletonData->RemoveComponent(ComponentIndex);
 			}
 
@@ -588,10 +591,13 @@ void UAnimationSharingManager::UnregisterAllActors()
 				FPerActorData& ActorData = Data->PerActorData[ActorIndex];
 				for (int32 ComponentIndex : ActorData.ComponentIndices)
 				{
-					Data->PerComponentData[ComponentIndex].Component->SetLeaderPoseComponent(nullptr, true);
-					Data->PerComponentData[ComponentIndex].Component->PrimaryComponentTick.bCanEverTick = true;
-					Data->PerComponentData[ComponentIndex].Component->SetComponentTickEnabled(true);
-					Data->PerComponentData[ComponentIndex].Component->bRecentlyRendered = false;
+					if (USkeletalMeshComponent* Component = Data->GetComponent(Data->PerComponentData[ComponentIndex]))
+					{
+						Component->SetLeaderPoseComponent(nullptr, true);
+						Component->PrimaryComponentTick.bCanEverTick = true;
+						Component->SetComponentTickEnabled(true);
+						Component->bRecentlyRendered = false;
+					}
 				}
 				ActorData.ComponentIndices.Empty();
 			}
@@ -667,7 +673,10 @@ void UAnimationSharingManager::SetDebugMaterialForActor(UAnimSharingInstance* Da
 #if DEBUG_MATERIALS 
 	for (uint32 ComponentIndex : Data->PerActorData[ActorIndex].ComponentIndices)
 	{
-		SetDebugMaterial(Data->PerComponentData[ComponentIndex].Component, State);
+		if (USkeletalMeshComponent* Component = Data->GetComponent(Data->PerComponentData[ComponentIndex]))
+		{
+			SetDebugMaterial(Component, State);
+		}
 	}
 #endif
 }
@@ -686,6 +695,18 @@ FName UAnimationSharingManager::GetPlatformName()
 }
 #endif
 
+USkeletalMeshComponent* UAnimSharingInstance::GetComponent(const FPerComponentData& Component)
+{
+	if (Component.WeakComponent.IsValid())
+	{
+		return Component.WeakComponent.Get();
+	}
+	
+	UE_LOG(LogAnimationSharing, Warning, TEXT("Stale SkeletalMeshComponent\n\tActor Index: %i\n\tActor: %s"), Component.ActorIndex, RegisteredActors[Component.ActorIndex].Get() ? *RegisteredActors[Component.ActorIndex]->GetName() : TEXT("None"));
+
+	return nullptr;
+}
+
 void UAnimSharingInstance::BeginDestroy()
 {
 	Super::BeginDestroy();
@@ -693,7 +714,10 @@ void UAnimSharingInstance::BeginDestroy()
 	{
 		for (uint32 ComponentIndex : ActorData.ComponentIndices)
 		{
-			PerComponentData[ComponentIndex].Component->SetLeaderPoseComponent(nullptr, true);
+			if (USkeletalMeshComponent* Component = GetComponent(PerComponentData[ComponentIndex]))
+			{
+				Component->SetLeaderPoseComponent(nullptr, true);
+			}
 		}
 	}
 
@@ -1090,11 +1114,15 @@ void UAnimSharingInstance::TickDebugInformation()
 
 				GEngine->AddOnScreenDebugMessage(1337, 1, FColor::White, OnScreenString);
 
-				const USkinnedMeshComponent* Component = PerComponentData[ActorData.ComponentIndices[0]].Component->LeaderPoseComponent.Get();
+				USkinnedMeshComponent* LeaderComponent = nullptr;
+				if (USkeletalMeshComponent* Component = GetComponent(PerComponentData[ActorData.ComponentIndices[0]]))
+				{
+					LeaderComponent = Component->LeaderPoseComponent.Get();
+				}
 #if UE_BUILD_DEVELOPMENT
-				if (Component != nullptr)
+				if (LeaderComponent != nullptr)
 				{
-					DrawDebugLine(GetWorld(), Actor->GetActorLocation(), Component->GetComponentLocation(), FColor::Magenta);
+					DrawDebugLine(GetWorld(), Actor->GetActorLocation(), LeaderComponent->GetComponentLocation(), FColor::Magenta);
 				}
 #endif
 			}
@@ -1345,10 +1373,13 @@ void UAnimSharingInstance::TickActorStates()
 			ActorData.bRequiresTick = ActorData.SignificanceValue >= ScalabilitySettings->TickSignificanceValue.Default;
 			for (int32 ComponentIndex : ActorData.ComponentIndices)
 			{
-				if (PerComponentData[ComponentIndex].Component->GetLastRenderTime() > (WorldTime - 1.f))
+				if (USkeletalMeshComponent* Component = GetComponent(PerComponentData[ComponentIndex]))
 				{
-					PerComponentData[ComponentIndex].Component->bRecentlyRendered = true;
-					ActorData.bRequiresTick = true;
+					if (Component->GetLastRenderTime() > (WorldTime - 1.f))
+					{
+						Component->bRecentlyRendered = true;
+						ActorData.bRequiresTick = true;
+					}
 				}
 			}
 
@@ -1549,7 +1580,10 @@ void UAnimSharingInstance::TickBlendInstances()
 
 					for (uint32 ComponentIndex : PerActorData[ActorIndex].ComponentIndices)
 					{
-						UAnimationSharingManager::SetDebugMaterial(PerComponentData[ComponentIndex].Component, 0);
+						if (USkeletalMeshComponent* Component = GetComponent(PerComponentData[ComponentIndex]))
+						{
+							UAnimationSharingManager::SetDebugMaterial(Component, 0);
+						}
 					}
 				}
 			}
@@ -1568,7 +1602,10 @@ void UAnimSharingInstance::TickBlendInstances()
 #if !UE_BUILD_SHIPPING
 						for (uint32 ComponentIndex : PerActorData[ActorIndex].ComponentIndices)
 						{
-							UAnimationSharingManager::SetDebugMaterial(PerComponentData[ComponentIndex].Component, 1);
+							if (USkeletalMeshComponent* Component = GetComponent(PerComponentData[ComponentIndex]))
+							{
+								UAnimationSharingManager::SetDebugMaterial(Component, 1);
+							}
 						}
 #endif
 					}
@@ -1601,7 +1638,10 @@ void UAnimSharingInstance::TickBlendInstances()
 
 					for (uint32 ComponentIndex : PerActorData[ActorIndex].ComponentIndices)
 					{
-						UAnimationSharingManager::SetDebugMaterial(PerComponentData[ComponentIndex].Component, 2);
+						if (USkeletalMeshComponent* Component = GetComponent(PerComponentData[ComponentIndex]))
+						{
+							UAnimationSharingManager::SetDebugMaterial(Component, 2);
+						}
 					}
 				}
 
@@ -1735,7 +1775,10 @@ void UAnimSharingInstance::SetLeaderComponentForActor(uint32 ActorIndex, USkelet
 	
 	for (uint32 ComponentIndex : ActorData.ComponentIndices)
 	{
-		PerComponentData[ComponentIndex].Component->SetLeaderPoseComponent(Component, true);
+		if (USkeletalMeshComponent* ActorComponent = GetComponent(PerComponentData[ComponentIndex]))
+		{
+			ActorComponent->SetLeaderPoseComponent(Component, true);
+		}
 	}
 }
 
diff --git a/Engine/Plugins/Developer/AnimationSharing/Source/AnimationSharing/Public/AnimationSharingManager.h b/Engine/Plugins/Developer/AnimationSharing/Source/AnimationSharing/Public/AnimationSharingManager.h
--- a/Engine/Plugins/Developer/AnimationSharing/Source/AnimationSharing/Public/AnimationSharingManager.h
+++ b/Engine/Plugins/Developer/AnimationSharing/Source/AnimationSharing/Public/AnimationSharingManager.h
@@ -229,7 +229,7 @@ struct FPerStateData
 struct FPerComponentData
 {
 	/** Skeletal mesh component registered for this component */
-	USkeletalMeshComponent* Component;
+	TWeakObjectPtr<USkeletalMeshComponent> WeakComponent;
 	/** Index to the owning actor (used to index PerActorData) */
 	int32 ActorIndex;
 };
@@ -396,6 +396,9 @@ public:
 	/** Sets the whether or not any of the follower components are visible */
 	void SetComponentTick(uint8 StateIndex, uint32 ComponentIndex);
 
+	/** Retrieves USkeletalMeshComponent with validation of reachability */
+	USkeletalMeshComponent* GetComponent(const FPerComponentData& ComponentData);
+
 	/** Actors currently registered to be animation driven by the AnimManager using this setup */
 	UPROPERTY(VisibleAnywhere, Transient, Category = AnimationSharing)
 	TArray<TObjectPtr<AActor>> RegisteredActors;
