AnimSharing - Added helper function for checking if an actor is registered and its internal data is valid and a function for logging out all the internal data of the manager. Debugging purpose only

diff --git a/Engine/Plugins/Developer/AnimationSharing/Source/AnimationSharing/Private/AnimationSharingManager.cpp b/Engine/Plugins/Developer/AnimationSharing/Source/AnimationSharing/Private/AnimationSharingManager.cpp
--- a/Engine/Plugins/Developer/AnimationSharing/Source/AnimationSharing/Private/AnimationSharingManager.cpp
+++ b/Engine/Plugins/Developer/AnimationSharing/Source/AnimationSharing/Private/AnimationSharingManager.cpp
@@ -2228,3 +2228,94 @@ void UAnimSharingInstance::RemoveAdditiveInstance(int32 InstanceIndex)
 
 	AdditiveInstances.RemoveAtSwap(InstanceIndex, EAllowShrinking::No);
 }
+
+bool UAnimationSharingManager::CheckDataForActor(AActor* InActor) const
+{
+	int32 PerSkeletonDataIdx = INDEX_NONE;
+	int32 RegisteredActorsIdx = INDEX_NONE;
+
+	for (int32 SkeletonIndex = 0; SkeletonIndex < PerSkeletonData.Num(); ++SkeletonIndex)
+	{
+		const UAnimSharingInstance* SkeletonData = PerSkeletonData[SkeletonIndex];
+		const int32 ActorIndex = SkeletonData->RegisteredActors.IndexOfByKey(InActor);
+
+		if (ActorIndex != INDEX_NONE)
+		{
+			PerSkeletonDataIdx = SkeletonIndex;
+			RegisteredActorsIdx = ActorIndex;
+			break;
+		}
+	}
+
+	if (PerSkeletonDataIdx == INDEX_NONE)
+	{
+		UE_LOG(LogAnimationSharing, Warning, TEXT("CheckActorData Failed! Reason: Actor %s is not registered. PerSkeletonData Num: %d"),
+			*GetNameSafe(InActor), PerSkeletonData.Num());
+		return false;
+	}
+
+	const UAnimSharingInstance* SkeletonData = PerSkeletonData[PerSkeletonDataIdx];
+	if (!SkeletonData->PerActorData.IsValidIndex(RegisteredActorsIdx))
+	{
+		UE_LOG(LogAnimationSharing, Warning, TEXT("CheckActorData Failed! Reason: ActorIndex '%d' for '%s' is not a valid index in PerActorData. SkeletonDataIdx: %d"),
+			RegisteredActorsIdx, *GetNameSafe(InActor), PerSkeletonDataIdx);
+		return false;
+	}
+
+	const FPerActorData& ActorData = SkeletonData->PerActorData[RegisteredActorsIdx];
+	if (ActorData.ComponentIndices.Num() == 0)
+	{
+		UE_LOG(LogAnimationSharing, Warning, TEXT("CheckActorData Failed! Reason: ComponentIndices for '%s' is empty. SkeletonDataIdx: %d ActorDataIdx: %d"),
+			*GetNameSafe(InActor), PerSkeletonDataIdx, RegisteredActorsIdx);
+		return false;
+	}
+
+	for (int32 ComponentIndex : ActorData.ComponentIndices)
+	{
+		if (!SkeletonData->PerComponentData.IsValidIndex(ComponentIndex))
+		{
+			UE_LOG(LogAnimationSharing, Warning, TEXT("CheckActorData Failed! Reason: ComponentIndex '%d' for '%s' is not a valid index in PerComponentData. PerComponentData.Num: %d"),
+				ComponentIndex, *GetNameSafe(InActor), SkeletonData->PerComponentData.Num());
+
+			return false;
+		}
+	}
+
+	return true;
+}
+
+void UAnimationSharingManager::LogData() const
+{
+	UE_LOG(LogAnimationSharing, Log, TEXT("PerSkeletonData Num: %d"), PerSkeletonData.Num());
+	for (int32 SkeletonIndex = 0; SkeletonIndex < PerSkeletonData.Num(); ++SkeletonIndex)
+	{
+		if (const UAnimSharingInstance* SkeletonData = PerSkeletonData[SkeletonIndex])
+		{
+			UE_LOG(LogAnimationSharing, Log, TEXT("\t RegisteredActors Num: %d"), SkeletonData->RegisteredActors.Num());
+			for (int32 Idx = 0; Idx < SkeletonData->RegisteredActors.Num(); Idx++)
+			{
+				UE_LOG(LogAnimationSharing, Log, TEXT("\t\t Idx: %d Actor: %s"), Idx, *GetNameSafe(SkeletonData->RegisteredActors[Idx]));
+			}
+
+			UE_LOG(LogAnimationSharing, Log, TEXT("\t PerActorData Num: %d"), SkeletonData->PerActorData.Num());
+			for (int32 Idx = 0; Idx < SkeletonData->PerActorData.Num(); Idx++)
+			{
+				const FPerActorData& ActorData = SkeletonData->PerActorData[Idx];
+				const FString ComponentIndicesStr = FString::JoinBy(ActorData.ComponentIndices, TEXT(", "), [](uint32 Item) { return FString::FromInt(Item); });
+
+				UE_LOG(LogAnimationSharing, Log, TEXT("\t\t Idx: %d CurrState: %d PrevState: %d PermutationIdx: %d ComponentIndices: %s"),
+					Idx, ActorData.CurrentState, ActorData.PreviousState, ActorData.PermutationIndex, *ComponentIndicesStr);
+			}
+
+			UE_LOG(LogAnimationSharing, Log, TEXT("\t PerComponentData Num: %d"), SkeletonData->PerComponentData.Num());
+			for (int32 Idx = 0; Idx < SkeletonData->PerComponentData.Num(); Idx++)
+			{
+				const FPerComponentData& PerComponentData = SkeletonData->PerComponentData[Idx];
+				const USkeletalMeshComponent* SkelMeshComp = PerComponentData.WeakComponent.Get();
+
+				UE_LOG(LogAnimationSharing, Log, TEXT("\t\t Idx: %d Comp: %s CompOwner: %s ActorIndex: %d"),
+					Idx, *GetNameSafe(SkelMeshComp), *GetNameSafe(SkelMeshComp ? SkelMeshComp->GetOwner() : nullptr), PerComponentData.ActorIndex);
+			}
+		}
+	}
+}
\ No newline at end of file
diff --git a/Engine/Plugins/Developer/AnimationSharing/Source/AnimationSharing/Public/AnimationSharingManager.h b/Engine/Plugins/Developer/AnimationSharing/Source/AnimationSharing/Public/AnimationSharingManager.h
--- a/Engine/Plugins/Developer/AnimationSharing/Source/AnimationSharing/Public/AnimationSharingManager.h
+++ b/Engine/Plugins/Developer/AnimationSharing/Source/AnimationSharing/Public/AnimationSharingManager.h
@@ -544,6 +544,13 @@ public:
 
 	void Tick(float DeltaTime);
 	FTickAnimationSharingFunction& GetTickFunction();
+
+	/** Checks whether the supplied actor is registered and all its data is valid */
+	bool CheckDataForActor(AActor* InActor) const;
+
+	/** Log out internal data */
+	void LogData() const;
+
 protected:
 	
 	/** Populates all data required for a Skeleton setup */
