Set some functions in AnimationSharingManager to virtual so they can be overridden.
Add support for creating and registering custom AnimationSharingManagers.
Add some more debug options to AnimationSharing system including showing current playing anims.

diff --git a/Engine/Plugins/Developer/AnimationSharing/Source/AnimationSharing/Private/AnimationSharingManager.cpp b/Engine/Plugins/Developer/AnimationSharing/Source/AnimationSharing/Private/AnimationSharingManager.cpp
--- a/Engine/Plugins/Developer/AnimationSharing/Source/AnimationSharing/Private/AnimationSharingManager.cpp
+++ b/Engine/Plugins/Developer/AnimationSharing/Source/AnimationSharing/Private/AnimationSharingManager.cpp
@@ -14,6 +14,11 @@
 #include "DrawDebugHelpers.h"
 #include "Engine/Engine.h"
 
+#if UE_BUILD_DEVELOPMENT
+#include "Animation/AnimSingleNodeInstance.h"
+#include "Animation/AnimNode_RelevantAssetPlayerBase.h"
+#endif // #if UE_BUILD_DEVELOPMENT
+
 #include UE_INLINE_GENERATED_CPP_BY_NAME(AnimationSharingManager)
 
 #if WITH_EDITOR
@@ -41,7 +46,23 @@ static FAutoConsoleVariableRef CVarAnimSharing_DebugStates(
 	TEXT("Controls whether and which animation sharing debug features are enabled.\n")
 	TEXT("0: Turned off.\n")
 	TEXT("1: Turns on active leader-components and blend with material coloring, and printing state information for each actor above their capsule.\n")
-	TEXT("2: Turns printing state information about currently active animation states, blend etc. Also enables line drawing from follower-components to currently assigned leader components."),
+	TEXT("2: Turns printing state information about currently active animation states, blend etc. Also enables line drawing from follower-components to currently assigned leader components.\n")
+	TEXT("3: Enable printing of extra state info, including currently playing animations (for anim blueprints this will only be the first relevant anim player on graph 'AnimGraph')."),
+	ECVF_Cheat);
+
+static int32 GAnimationSharingDisableDebugMaterials = 0;
+static FAutoConsoleVariableRef CVarAnimSharing_DisableDebugMaterials(
+	TEXT("a.Sharing.DisableDebugMaterials"),
+	GAnimationSharingDisableDebugMaterials,
+	TEXT("Values: 0/1\n")
+	TEXT("Controls whether material coloring should be disabled when anim sharing debug is enabled.\n"),
+	ECVF_Cheat);
+
+static float GAnimationSharingDebugFontScale = 1.0f;
+static FAutoConsoleVariableRef CVarAnimSharing_FontScale(
+	TEXT("a.Sharing.DebugFontScale"),
+	GAnimationSharingDebugFontScale,
+	TEXT("Scale of font used for anim sharing debug.\n"),
 	ECVF_Cheat);
 
 static int32 GAnimationSharingEnabled = 1;
@@ -211,6 +232,11 @@ const FAnimationSharingScalability& UAnimationSharingManager::GetScalabilitySett
 	return ScalabilitySettings;
 }
 
+UAnimSharingInstance* UAnimationSharingManager::CreateAnimSharingInstance()
+{
+	return NewObject<UAnimSharingInstance>(this);
+}
+
 void UAnimationSharingManager::SetupPerSkeletonData(const FPerSkeletonAnimationSharingSetup& SkeletonSetup)
 {
 	const USkeleton* Skeleton = SkeletonSetup.Skeleton;
@@ -218,7 +244,8 @@ void UAnimationSharingManager::SetupPerSkeletonData(const FPerSkeletonAnimationS
 	UEnum* StateEnum = Processor ? Processor->GetAnimationStateEnum() : nullptr;
 	if (Skeleton && StateEnum && Processor)
 	{
-		UAnimSharingInstance* Data = NewObject<UAnimSharingInstance>(this);
+		UAnimSharingInstance* Data = CreateAnimSharingInstance();
+		
 		// Try and setup up instance using provided setup data
 		if (Data->Setup(this, SkeletonSetup, &ScalabilitySettings, Skeletons.Num()))
 		{
@@ -657,7 +684,7 @@ bool UAnimationSharingManager::CreateAnimationSharingManager(UObject* WorldConte
 void UAnimationSharingManager::SetDebugMaterial(USkeletalMeshComponent* Component, uint8 State)
 {
 #if DEBUG_MATERIALS 
-	if (GAnimationSharingDebugging >= 1 && DebugMaterials.IsValidIndex(State))
+	if (GAnimationSharingDebugging >= 1 && GAnimationSharingDisableDebugMaterials == 0 && DebugMaterials.IsValidIndex(State))
 	{
 		const int32 NumMaterials = Component->GetNumMaterials();
 		for (int32 Index = 0; Index < NumMaterials; ++Index)
@@ -1038,6 +1065,83 @@ void UAnimSharingInstance::SetupState(FPerStateData& StateData, const FAnimation
 	StateData.FollowerTickRequiredFrameBits.Init(false, Components.Num());
 }
 
+#if UE_BUILD_DEVELOPMENT
+void UAnimSharingInstance::GetAnimAssetNameAndData(const uint32 StateIdx, const uint32 CompIdx, FString& AnimAssetName, float& AnimTime, float& AnimLength) const
+{
+	if (PerStateData[StateIdx].Components.IsValidIndex(CompIdx))
+	{
+		USkeletalMeshComponent* Comp = PerStateData[StateIdx].Components[CompIdx];
+		const EAnimationMode::Type AnimMode = Comp->GetAnimationMode();
+		switch (AnimMode)
+		{
+			case EAnimationMode::AnimationSingleNode:
+			{
+				if (UAnimSingleNodeInstance* SingleNodeAnimInstance = Comp->GetSingleNodeInstance())
+				{
+					if (UAnimationAsset* AnimAsset = SingleNodeAnimInstance->GetAnimationAsset())
+					{
+						AnimAssetName = AnimAsset->GetName();
+					}
+					AnimTime = SingleNodeAnimInstance->GetCurrentTime();
+					AnimLength = SingleNodeAnimInstance->GetLength();
+				}
+			}
+			break;
+			
+			case EAnimationMode::AnimationBlueprint:
+			{
+				if (UAnimInstance* AnimInstance = Comp->GetAnimInstance())
+				{
+					AnimAssetName = AnimInstance->GetName();
+
+					static FName AnimGraphName(TEXT("AnimGraph"));
+					TArray<const FAnimNode_AssetPlayerRelevancyBase*> PlayerArray = AnimInstance->GetInstanceRelevantAssetPlayers(AnimGraphName);
+					if (PlayerArray.Num() > 0)
+					{
+						if (UAnimationAsset* AnimAsset = PlayerArray[0]->GetAnimAsset())
+						{
+							AnimAssetName = AnimAsset->GetName();
+						}
+						AnimTime = PlayerArray[0]->GetCurrentAssetTime();
+						AnimLength = PlayerArray[0]->GetCurrentAssetLength();
+					}
+				}
+			}
+			break;
+
+			default:
+				break;
+		}
+	}
+}
+
+FString UAnimSharingInstance::GetDebugStateString(uint8 State, const uint32 CompIdx, uint32 OnDemandInstanceIndex) const
+{
+	if (OnDemandInstanceIndex != INDEX_NONE && OnDemandInstances.IsValidIndex(OnDemandInstanceIndex))
+	{
+		const FOnDemandInstance& OnDemandInstance = OnDemandInstances[OnDemandInstanceIndex];
+		const uint32 OnDemandStateIdx = OnDemandInstance.State;
+		const uint32 OnDemandCompIdx = OnDemandInstance.UsedPerStateComponentIndex;
+
+		FString OnDemandAnimName;
+		float OnDemandAnimTime = 0.0f;
+		float OnDemandAnimLength = 0.0f;
+		GetAnimAssetNameAndData(OnDemandStateIdx, OnDemandCompIdx, OnDemandAnimName, OnDemandAnimTime, OnDemandAnimLength);
+
+		return FString::Printf(TEXT("On Demand [%i] Anim %s [%.2f / %.2f] {%d}"), OnDemandInstanceIndex, *OnDemandAnimName, OnDemandAnimTime, OnDemandAnimLength, OnDemandCompIdx);
+	}
+	else
+	{
+		FString AnimName;
+		float AnimTime = 0.0f;
+		float AnimLength = 0.0f;
+		GetAnimAssetNameAndData(State, CompIdx, AnimName, AnimTime, AnimLength);
+
+		return FString::Printf(TEXT("%s [%.2f / %.2f] {%d}"), *AnimName, AnimTime, AnimLength, CompIdx);
+	}
+}
+#endif // #if UE_BUILD_DEVELOPMENT
+
 void UAnimSharingInstance::TickDebugInformation()
 {
 #if !UE_BUILD_SHIPPING
@@ -1055,7 +1159,6 @@ void UAnimSharingInstance::TickDebugInformation()
 	}
 #endif // UE_BUILD_DEVELOPMENT
 	
-
 	for (int32 ActorIndex = 0; ActorIndex < RegisteredActors.Num(); ++ActorIndex)
 	{
 		// Non-const for DrawDebugString
@@ -1063,8 +1166,28 @@ void UAnimSharingInstance::TickDebugInformation()
 		if (Actor)
 		{
 			const FPerActorData& ActorData = PerActorData[ActorIndex];
+
+#if UE_BUILD_DEVELOPMENT
 			const uint8 State = ActorData.CurrentState;
 
+			const FColor DebugColor = [&ActorData]()
+			{
+				const uint32 BlendInstanceIndex = ActorData.BlendInstanceIndex;
+				const uint32 DemandInstanceIndex = ActorData.OnDemandInstanceIndex;
+
+				/** Colors match debug material colors */
+				if (ActorData.bBlending && BlendInstanceIndex != INDEX_NONE)
+				{
+					return FColor::Blue;
+				}
+				else if (ActorData.bRunningOnDemand && DemandInstanceIndex != INDEX_NONE)
+				{
+					return FColor::Red;
+				}
+
+				return FColor::Green;
+			}();
+
 			const FString StateString = [this, &ActorData, State]() -> FString
 			{
 				/** Check whether or not we are currently blending */
@@ -1086,27 +1209,36 @@ void UAnimSharingInstance::TickDebugInformation()
 				return FString::Printf(TEXT("State - %s %1.2f"), *StateEnum->GetDisplayNameTextByValue(State).ToString(), ActorData.SignificanceValue);
 			}();
 
-			const FColor DebugColor = [&ActorData]()
+			const FString FullStateString = [this, &ActorData, State]()
 			{
+				/** Check whether or not we are currently blending */
 				const uint32 BlendInstanceIndex = ActorData.BlendInstanceIndex;
-				const uint32 DemandInstanceIndex = ActorData.OnDemandInstanceIndex;
-
-				/** Colors match debug material colors */				
-				if (ActorData.bBlending && BlendInstanceIndex != INDEX_NONE)
+				if (BlendInstanceIndex != INDEX_NONE && BlendInstances.IsValidIndex(BlendInstanceIndex))
 				{
-					return FColor::Blue;
+					const FBlendInstance& BlendInstance = BlendInstances[BlendInstanceIndex];
+					uint8 StateFrom = BlendInstance.StateFrom;
+					uint8 StateTo = BlendInstance.StateTo;
+					const float TimeLeft = BlendInstances[BlendInstanceIndex].EndTime - static_cast<float>(GetWorld()->GetTimeSeconds());
+
+					return FString::Printf(TEXT("Blending states - %s to %s [%1.3f] (%i)\n  From: %s\n  To: %s"),
+						*StateEnum->GetDisplayNameTextByValue(StateFrom).ToString(),
+						*StateEnum->GetDisplayNameTextByValue(StateTo).ToString(),
+						TimeLeft, ActorData.BlendInstanceIndex, 
+						*GetDebugStateString(StateFrom, BlendInstance.FromPermutationIndex, BlendInstance.FromOnDemandInstanceIndex),
+						*GetDebugStateString(StateTo, BlendInstance.ToPermutationIndex, BlendInstance.ToOnDemandInstanceIndex)
+					);
 				}
-				else if (ActorData.bRunningOnDemand && DemandInstanceIndex != INDEX_NONE)
+				else
 				{
-					return FColor::Red;
+					/** Otherwise we should just be part of a state */
+					return FString::Printf(TEXT("State - %s Sig %1.2f\n  %s"), *StateEnum->GetDisplayNameTextByValue(State).ToString(),
+							ActorData.SignificanceValue, *GetDebugStateString(State, ActorData.PermutationIndex, ActorData.OnDemandInstanceIndex));
 				}
-
-				return FColor::Green;
 			}();
 
-#if UE_BUILD_DEVELOPMENT
 			/** Draw text above AI pawn's head */
-			DrawDebugString(GetWorld(), FVector(0.f, 0.f, 100.f), StateString, Actor, DebugColor, 0.016f, false);
+			DrawDebugString(GetWorld(), FVector(0.f, 0.f, 100.f), (GAnimationSharingDebugging >= 3) ? FullStateString : StateString, 
+				Actor, GAnimationSharingDisableDebugMaterials ? FColor::Green : DebugColor, 0.016f, false, GAnimationSharingDebugFontScale);
 #endif
 			if (GAnimationSharingDebugging >= 2)
 			{
diff --git a/Engine/Plugins/Developer/AnimationSharing/Source/AnimationSharing/Private/AnimationSharingModule.cpp b/Engine/Plugins/Developer/AnimationSharing/Source/AnimationSharing/Private/AnimationSharingModule.cpp
--- a/Engine/Plugins/Developer/AnimationSharing/Source/AnimationSharing/Private/AnimationSharingModule.cpp
+++ b/Engine/Plugins/Developer/AnimationSharing/Source/AnimationSharing/Private/AnimationSharingModule.cpp
@@ -31,13 +31,22 @@ void FAnimSharingModule::AddReferencedObjects(FReferenceCollector& Collector)
 
 bool FAnimSharingModule::CreateAnimationSharingManager(UWorld* InWorld, const UAnimationSharingSetup* Setup)
 {
-	if (InWorld && InWorld->IsGameWorld() && Setup && UAnimationSharingManager::AnimationSharingEnabled() && !WorldAnimSharingManagers.Contains(InWorld))
+	return CreateAnimationSharingManager(InWorld, [Setup] (UWorld* InWorld)
+		{
+			UAnimationSharingManager* Manager = NewObject<UAnimationSharingManager>(InWorld);
+			Manager->Initialise(Setup);
+			return Manager;
+		} 
+	);
+}
+
+bool FAnimSharingModule::CreateAnimationSharingManager(UWorld* InWorld, FCreateAnimationSharingManagerFunc CreateAnimationSharingManagerFunc)
+{
+	if (InWorld && InWorld->IsGameWorld() && CreateAnimationSharingManagerFunc && UAnimationSharingManager::AnimationSharingEnabled() && !WorldAnimSharingManagers.Contains(InWorld))
 	{
 		UE_SCOPED_ENGINE_ACTIVITY(TEXT("Initializing Animation Sharing"));
-		UAnimationSharingManager* Manager = NewObject<UAnimationSharingManager>(InWorld);
-		Manager->Initialise(Setup);
-		WorldAnimSharingManagers.Add(InWorld, Manager);
-		
+		UAnimationSharingManager* Manager = CreateAnimationSharingManagerFunc(InWorld);
+		WorldAnimSharingManagers.Add(InWorld, Manager);	
 		OnAnimationSharingManagerCreated.Broadcast(Manager, InWorld);
 
 		return true;
diff --git a/Engine/Plugins/Developer/AnimationSharing/Source/AnimationSharing/Public/AnimationSharingManager.h b/Engine/Plugins/Developer/AnimationSharing/Source/AnimationSharing/Public/AnimationSharingManager.h
--- a/Engine/Plugins/Developer/AnimationSharing/Source/AnimationSharing/Public/AnimationSharingManager.h
+++ b/Engine/Plugins/Developer/AnimationSharing/Source/AnimationSharing/Public/AnimationSharingManager.h
@@ -271,9 +271,8 @@ struct FPerActorData
 	FUpdateActorHandle UpdateActorHandleDelegate;	
 };
 
-
 UCLASS()
-class UAnimSharingInstance : public UObject
+class ANIMATIONSHARING_API UAnimSharingInstance : public UObject
 {
 	GENERATED_BODY()
 
@@ -388,7 +387,7 @@ public:
 	void SetPermutationSlaveComponent(uint8 StateIndex, uint32 ActorIndex, uint32 PermutationIndex) { SetPermutationFollowerComponent(StateIndex, ActorIndex, PermutationIndex); }
 	
 	/** Determines a permutation index for the given actor and state */
-	uint32 DeterminePermutationIndex(uint32 ActorIndex, uint8 State) const;
+	virtual uint32 DeterminePermutationIndex(uint32 ActorIndex, uint8 State) const;
 
 	/** Marks the component as either used/not-used, this is used to disable ticking of components which are not in use*/
 	void SetComponentUsage(bool bUsage, uint8 StateIndex, uint32 ComponentIndex);
@@ -454,6 +453,14 @@ public:
 
 	/** Holds the current frame world time */
 	float WorldTime;
+
+private:
+
+#if UE_BUILD_DEVELOPMENT
+	void GetAnimAssetNameAndData(const uint32 StateIdx, const uint32 CompIdx, FString& AnimAssetName, float& AnimTime, float& AnimLength) const;
+	FString GetDebugStateString(uint8 State, const uint32 CompIdx, uint32 OnDemandInstanceIndex) const;
+#endif // #if UE_BUILD_DEVELOPMENT
+
 };
 
 USTRUCT()
@@ -556,6 +563,8 @@ protected:
 	/** Populates all data required for a Skeleton setup */
 	void SetupPerSkeletonData(const FPerSkeletonAnimationSharingSetup& SkeletonSetup);
 
+	virtual UAnimSharingInstance* CreateAnimSharingInstance();
+
 	/** Dealing with Actor data and handles */
 	uint32 CreateActorHandle(uint8 SkeletonIndex, uint32 ActorIndex) const;
 	uint8 GetSkeletonIndexFromHandle(uint32 InHandle) const;
diff --git a/Engine/Plugins/Developer/AnimationSharing/Source/AnimationSharing/Public/AnimationSharingModule.h b/Engine/Plugins/Developer/AnimationSharing/Source/AnimationSharing/Public/AnimationSharingModule.h
--- a/Engine/Plugins/Developer/AnimationSharing/Source/AnimationSharing/Public/AnimationSharingModule.h
+++ b/Engine/Plugins/Developer/AnimationSharing/Source/AnimationSharing/Public/AnimationSharingModule.h
@@ -37,6 +37,12 @@ public:
 
 	/** Creates an animation sharing manager for the given UWorld (must be a Game World) */
 	static bool CreateAnimationSharingManager(UWorld* InWorld, const UAnimationSharingSetup* Setup);
+	
+	/** Creates an animation sharing manager for the given UWorld (must be a Game World) - with 
+		custom func that is responsible for creation and initialisation */
+	using FCreateAnimationSharingManagerFunc = TFunction<UAnimationSharingManager*(UWorld* InWorld)>;
+	static bool CreateAnimationSharingManager(UWorld* InWorld, FCreateAnimationSharingManagerFunc CreateAnimationSharingManagerFunc);
+
 private:	
 	static void OnWorldCleanup(UWorld* World, bool bSessionEnded, bool bCleanupResources);
 	static TMap<const UWorld*, TObjectPtr<UAnimationSharingManager>> WorldAnimSharingManagers;
