AnimationSharingManager can now create sharing instances from multiple AnimationSharingSetups.
AnimationSharingSetups can be added and removed from the AnimationaSharingManager as required.

diff --git a/Engine/Plugins/Developer/AnimationSharing/Source/AnimationSharing/Private/AnimationSharingManager.cpp b/Engine/Plugins/Developer/AnimationSharing/Source/AnimationSharing/Private/AnimationSharingManager.cpp
--- a/Engine/Plugins/Developer/AnimationSharing/Source/AnimationSharing/Private/AnimationSharingManager.cpp
+++ b/Engine/Plugins/Developer/AnimationSharing/Source/AnimationSharing/Private/AnimationSharingManager.cpp
@@ -138,6 +138,9 @@ void UAnimationSharingManager::BeginDestroy()
 {
 	Super::BeginDestroy();
 	PerSkeletonData.Empty();
+	AnimationSharingSetups.Empty();
+	SetupSkeletonDataRanges.Empty();
+	ScalabilitySettingsPerSetup.Empty();
 	
 	// Unregister tick function
 	TickFunction.UnRegisterTickFunction();
@@ -190,46 +193,134 @@ FTickAnimationSharingFunction& UAnimationSharingManager::GetTickFunction()
 
 void UAnimationSharingManager::Initialise(const UAnimationSharingSetup* InSetup)
 {
+	TickFunction.Manager = this;
+	TickFunction.RegisterTickFunction(GetWorld()->PersistentLevel);
+
+	// Debug materials
+#if DEBUG_MATERIALS 
+	DebugMaterials.Empty();
+	{
+		UMaterialInterface* RedMaterial = LoadObject<UMaterialInterface>(nullptr, TEXT("/AnimationSharing/AnimSharingRed.AnimSharingRed"));
+		DebugMaterials.Add(RedMaterial);
+		UMaterialInterface* GreenMaterial = LoadObject<UMaterialInterface>(nullptr, TEXT("/AnimationSharing/AnimSharingGreen.AnimSharingGreen"));
+		DebugMaterials.Add(GreenMaterial);
+		UMaterialInterface* BlueMaterial = LoadObject<UMaterialInterface>(nullptr, TEXT("/AnimationSharing/AnimSharingBlue.AnimSharingBlue"));
+		DebugMaterials.Add(BlueMaterial);
+	}
+#endif 
+
+	AddSetup(InSetup);
+}
+
+void UAnimationSharingManager::AddSetup(const UAnimationSharingSetup* InSetup)
+{
+	// check setup doesn't already exist
 	if (InSetup)
 	{
-		TickFunction.Manager = this;
-		TickFunction.RegisterTickFunction(GetWorld()->PersistentLevel);
+		if (!AnimationSharingSetups.Contains(InSetup))
+		{
+			AnimationSharingSetups.Add(InSetup);
 
-		ScalabilitySettings = InSetup->ScalabilitySettings;
+			FAnimationSharingScalability& NewScalabilitySettings = ScalabilitySettingsPerSetup.Emplace_GetRef(InSetup->ScalabilitySettings);
 
-#if WITH_EDITOR
-		// Update local copy defaults with current platform value
-		const FName PlatformName = UAnimationSharingManager::GetPlatformName();
-		ScalabilitySettings.UseBlendTransitions = ScalabilitySettings.UseBlendTransitions.GetValueForPlatform(PlatformName);
-		ScalabilitySettings.BlendSignificanceValue = ScalabilitySettings.BlendSignificanceValue.GetValueForPlatform(PlatformName);
-		ScalabilitySettings.MaximumNumberConcurrentBlends = ScalabilitySettings.MaximumNumberConcurrentBlends.GetValueForPlatform(PlatformName);
-		ScalabilitySettings.TickSignificanceValue = ScalabilitySettings.TickSignificanceValue.GetValueForPlatform(PlatformName);
-#endif
+	#if WITH_EDITOR
+			// Update local copy defaults with current platform value
+			const FName PlatformName = UAnimationSharingManager::GetPlatformName();
+			NewScalabilitySettings.UseBlendTransitions = NewScalabilitySettings.UseBlendTransitions.GetValueForPlatform(PlatformName);
+			NewScalabilitySettings.BlendSignificanceValue = NewScalabilitySettings.BlendSignificanceValue.GetValueForPlatform(PlatformName);
+			NewScalabilitySettings.MaximumNumberConcurrentBlends = NewScalabilitySettings.MaximumNumberConcurrentBlends.GetValueForPlatform(PlatformName);
+			NewScalabilitySettings.TickSignificanceValue = NewScalabilitySettings.TickSignificanceValue.GetValueForPlatform(PlatformName);
+	#endif
 
-		// Debug materials
-#if DEBUG_MATERIALS 
-		DebugMaterials.Empty();
+			SetupSkeletonDataRanges.Emplace(PerSkeletonData.Num(), InSetup->SkeletonSetups.Num());
+			
+			for (const FPerSkeletonAnimationSharingSetup& SkeletonSetup : InSetup->SkeletonSetups)
+			{
+				SetupPerSkeletonData(SkeletonSetup, &NewScalabilitySettings);
+			}
+		}
+		else
 		{
-			UMaterialInterface* RedMaterial = LoadObject<UMaterialInterface>(nullptr, TEXT("/AnimationSharing/AnimSharingRed.AnimSharingRed"));
-			DebugMaterials.Add(RedMaterial);
-			UMaterialInterface* GreenMaterial = LoadObject<UMaterialInterface>(nullptr, TEXT("/AnimationSharing/AnimSharingGreen.AnimSharingGreen"));
-			DebugMaterials.Add(GreenMaterial);
-			UMaterialInterface* BlueMaterial = LoadObject<UMaterialInterface>(nullptr, TEXT("/AnimationSharing/AnimSharingBlue.AnimSharingBlue"));
-			DebugMaterials.Add(BlueMaterial);
+			UE_LOG(LogAnimationSharing, Error, TEXT("Passed duplicate UAnimationSharingSetup '%s' to add to AnimationSharingManager!"), *InSetup->GetName());
 		}
-#endif 
-		UWorld* World = GetWorld();
+	}
+}
 
-		for (const FPerSkeletonAnimationSharingSetup& SkeletonSetup : InSetup->SkeletonSetups)
+void UAnimationSharingManager::RemoveSetup(const UAnimationSharingSetup* InSetup)
+{
+	if (InSetup)
+	{
+		const int32 SetupIdx = AnimationSharingSetups.IndexOfByKey(InSetup);
+		if (SetupIdx != INDEX_NONE)
 		{
-			SetupPerSkeletonData(SkeletonSetup);
+			check(ScalabilitySettingsPerSetup.IsValidIndex(SetupIdx));
+			check(SetupSkeletonDataRanges.IsValidIndex(SetupIdx));
+
+			const FAnimationSharingSkeletonDataRange& AnimationSharingSetupRange = SetupSkeletonDataRanges[SetupIdx];
+
+			const int32 StartIndex = AnimationSharingSetupRange.Key;
+			const int32 NumSkeletonDatas = AnimationSharingSetupRange.Value;
+			const int32 EndIndex = StartIndex + NumSkeletonDatas;
+
+			check(PerSkeletonData.IsValidIndex(StartIndex));
+			check(Skeletons.IsValidIndex(StartIndex));
+			check(EndIndex <= PerSkeletonData.Num());
+			check(EndIndex <= Skeletons.Num());
+
+			for (int32 PerSkeletonDataIdx = StartIndex; PerSkeletonDataIdx < EndIndex ; ++PerSkeletonDataIdx)
+			{
+				if (PerSkeletonData[PerSkeletonDataIdx] != nullptr)
+				{
+					// check no associated actors are registered
+					if (PerSkeletonData[PerSkeletonDataIdx]->RegisteredActors.Num() > 0)
+					{
+						UE_LOG(LogAnimationSharing, Error, TEXT("There are still actors registered to skeleton '%s' when animation sharing setup '%s' was removed."), 
+							(Skeletons[PerSkeletonDataIdx]) ? *Skeletons[PerSkeletonDataIdx]->GetName() : TEXT("NULL"), 
+							*InSetup->GetName());
+					}
+				}
+			}
+
+			// remove associated skeletons and perskeletondata
+			PerSkeletonData.RemoveAt(StartIndex, NumSkeletonDatas);
+			Skeletons.RemoveAt(StartIndex, NumSkeletonDatas);
+
+			// remove scalability settings and range data
+			SetupSkeletonDataRanges.RemoveAt(SetupIdx);
+			ScalabilitySettingsPerSetup.RemoveAt(SetupIdx);
+
+			AnimationSharingSetups.Remove(InSetup);
+		}
+		else
+		{
+			UE_LOG(LogAnimationSharing, Error, TEXT("Passed duplicate UAnimationSharingSetup '%s' to remove from AnimationSharingManager!"), *InSetup->GetName());
 		}
 	}
+	else
+	{
+		UE_LOG(LogAnimationSharing, Error, TEXT("Passed null UAnimationSharingSetup to remove from AnimationSharingManager!"));
+	}
+
 }
 
-const FAnimationSharingScalability& UAnimationSharingManager::GetScalabilitySettings() const
+const FAnimationSharingScalability& UAnimationSharingManager::GetScalabilitySettings(const UAnimationSharingSetup* InSetup /*= nullptr*/) const
 {
-	return ScalabilitySettings;
+	if (InSetup != nullptr)
+	{
+		const int32 SetupIdx = AnimationSharingSetups.Find(InSetup);
+		if (SetupIdx != INDEX_NONE)
+		{
+			check(SetupIdx < ScalabilitySettingsPerSetup.Num());
+			return ScalabilitySettingsPerSetup[SetupIdx];
+		}
+	}
+	else if (!ScalabilitySettingsPerSetup.IsEmpty())
+	{
+		// if no setup parameter was specified, return first in list, else drop down to default settings.
+		return ScalabilitySettingsPerSetup[0];
+	}
+
+	return DefaultScalabilitySettings;
 }
 
 UAnimSharingInstance* UAnimationSharingManager::CreateAnimSharingInstance()
@@ -237,7 +328,7 @@ UAnimSharingInstance* UAnimationSharingManager::CreateAnimSharingInstance()
 	return NewObject<UAnimSharingInstance>(this);
 }
 
-void UAnimationSharingManager::SetupPerSkeletonData(const FPerSkeletonAnimationSharingSetup& SkeletonSetup)
+void UAnimationSharingManager::SetupPerSkeletonData(const FPerSkeletonAnimationSharingSetup& SkeletonSetup, const FAnimationSharingScalability* ScalabilitySettingsForSkeleton)
 {
 	const USkeleton* Skeleton = SkeletonSetup.Skeleton;
 	UAnimationSharingStateProcessor* Processor = SkeletonSetup.StateProcessorClass ?SkeletonSetup.StateProcessorClass->GetDefaultObject<UAnimationSharingStateProcessor>() : nullptr;
@@ -247,7 +338,7 @@ void UAnimationSharingManager::SetupPerSkeletonData(const FPerSkeletonAnimationS
 		UAnimSharingInstance* Data = CreateAnimSharingInstance();
 		
 		// Try and setup up instance using provided setup data
-		if (Data->Setup(this, SkeletonSetup, &ScalabilitySettings, Skeletons.Num()))
+		if (Data->Setup(this, SkeletonSetup, ScalabilitySettingsForSkeleton, Skeletons.Num()))
 		{
 			PerSkeletonData.Add(Data);
 			Skeletons.Add(Skeleton);
diff --git a/Engine/Plugins/Developer/AnimationSharing/Source/AnimationSharing/Public/AnimationSharingManager.h b/Engine/Plugins/Developer/AnimationSharing/Source/AnimationSharing/Public/AnimationSharingManager.h
--- a/Engine/Plugins/Developer/AnimationSharing/Source/AnimationSharing/Public/AnimationSharingManager.h
+++ b/Engine/Plugins/Developer/AnimationSharing/Source/AnimationSharing/Public/AnimationSharingManager.h
@@ -548,8 +548,14 @@ public:
 	/** Initialize sharing data structures */
 	UE_API void Initialise(const UAnimationSharingSetup* InSetup);
 
-	/** Returns current scalability settings */
-	UE_API const FAnimationSharingScalability& GetScalabilitySettings() const;
+	/** Add an animation sharing setup to the manager */
+	UE_API void AddSetup(const UAnimationSharingSetup* InSetup);
+
+	/** Remove an animation sharing setup from the manager */
+	UE_API void RemoveSetup(const UAnimationSharingSetup* InSetup);
+
+	/** Returns current scalability settings, with optional setup parameter to return scalability settings specific to an animation sharing setup */
+	UE_API const FAnimationSharingScalability& GetScalabilitySettings(const UAnimationSharingSetup* InSetup = nullptr) const;
 
 	UE_API void Tick(float DeltaTime);
 	UE_API FTickAnimationSharingFunction& GetTickFunction();
@@ -563,7 +569,7 @@ public:
 protected:
 	
 	/** Populates all data required for a Skeleton setup */
-	UE_API void SetupPerSkeletonData(const FPerSkeletonAnimationSharingSetup& SkeletonSetup);
+	UE_API void SetupPerSkeletonData(const FPerSkeletonAnimationSharingSetup& SkeletonSetup, const FAnimationSharingScalability* ScalabilitySettingsForSkeleton);
 
 	UE_API virtual UAnimSharingInstance* CreateAnimSharingInstance();
 
@@ -581,8 +587,20 @@ protected:
 	UPROPERTY(VisibleAnywhere, Transient, Category = AnimationSharing)
 	TArray<TObjectPtr<UAnimSharingInstance>> PerSkeletonData;
 	
-	/** Platform specific scalability settings */
-	FAnimationSharingScalability ScalabilitySettings;
+	/** Pointers to animation sharing setups added to manager */
+	UPROPERTY(Transient)
+	TArray< TWeakObjectPtr<const UAnimationSharingSetup> > AnimationSharingSetups;
+
+	/** Platform specific scalability settings (default version, will be overridden by a setup) */
+	FAnimationSharingScalability DefaultScalabilitySettings;
+
+	using FAnimationSharingSkeletonDataRange = TPair < uint32, uint32 >;  // Key = StartIndex, Value = NumDatas
+
+	/** List of ranges for which perskeletondata and skeletons correspond to each animation sharing setup */
+	TArray<FAnimationSharingSkeletonDataRange> SetupSkeletonDataRanges;
+
+	/** List of scalability settings specific to each animation sharing setup */
+	TArray<FAnimationSharingScalability> ScalabilitySettingsPerSetup;
 
 	/** Tick function for this manager */
 	FTickAnimationSharingFunction TickFunction;
