PR #13339: Expose Unregister Actor in AnimSharing

diff --git a/Engine/Plugins/Developer/AnimationSharing/Source/AnimationSharing/Public/AnimationSharingManager.h b/Engine/Plugins/Developer/AnimationSharing/Source/AnimationSharing/Public/AnimationSharingManager.h
--- a/Engine/Plugins/Developer/AnimationSharing/Source/AnimationSharing/Public/AnimationSharingManager.h
+++ b/Engine/Plugins/Developer/AnimationSharing/Source/AnimationSharing/Public/AnimationSharingManager.h
@@ -514,6 +514,12 @@ public:
 	UFUNCTION(BlueprintCallable, Category = AnimationSharing, meta = (DisplayName = "Register Actor"))
 	UE_API void RegisterActorWithSkeletonBP(AActor* InActor, const USkeleton* SharingSkeleton);
 
+	UFUNCTION(BlueprintCallable, Category = AnimationSharing, meta = (DisplayName = "Unregister Actor"))
+	void UnregisterActorBP(AActor* InActor)
+	{
+		UnregisterActor(InActor);
+	}
+
 	/** Returns whether or not the animation sharing is enabled */
 	UFUNCTION(BlueprintPure, Category = AnimationSharing)
 	static UE_API bool AnimationSharingEnabled();
