Fix for incorrect use of skeletonIndex rather than skeletonID in UAnimationSharingManager::UnregisterActor.

diff --git a/Engine/Plugins/Developer/AnimationSharing/Source/AnimationSharing/Private/AnimationSharingManager.cpp b/Engine/Plugins/Developer/AnimationSharing/Source/AnimationSharing/Private/AnimationSharingManager.cpp
--- a/Engine/Plugins/Developer/AnimationSharing/Source/AnimationSharing/Private/AnimationSharingManager.cpp
+++ b/Engine/Plugins/Developer/AnimationSharing/Source/AnimationSharing/Private/AnimationSharingManager.cpp
@@ -691,7 +691,7 @@ void UAnimationSharingManager::UnregisterActor(AActor* InActor)
 				}
 
 				// Make sure we update the handle on the swapped actor
-				SkeletonData->PerActorData[SwapIndex].UpdateActorHandleDelegate.ExecuteIfBound(CreateActorHandle(IntCastChecked<uint8>(SkeletonIndex), ActorIndex));
+				SkeletonData->PerActorData[SwapIndex].UpdateActorHandleDelegate.ExecuteIfBound(CreateActorHandle(IntCastChecked<uint8>(SkeletonData->SkeletonID), ActorIndex));
 			}			
 
 			SkeletonData->PerActorData.RemoveAtSwap(ActorIndex, EAllowShrinking::No);
