PermutationTimeOffset requiring InitAnim for AnimSharing instances
#fix add InitAnim after initializing AnimationSharing anim-instance(s) to ensure property values are updated correctly

diff --git a/Engine/Plugins/Developer/AnimationSharing/Source/AnimationSharing/Private/AnimationSharingManager.cpp b/Engine/Plugins/Developer/AnimationSharing/Source/AnimationSharing/Private/AnimationSharingManager.cpp
--- a/Engine/Plugins/Developer/AnimationSharing/Source/AnimationSharing/Private/AnimationSharingManager.cpp
+++ b/Engine/Plugins/Developer/AnimationSharing/Source/AnimationSharing/Private/AnimationSharingManager.cpp
@@ -1127,6 +1127,8 @@ void UAnimSharingInstance::SetupState(FPerStateData& StateData, const FAnimation
 
 							/** Set the current animation length length */
 							StateData.AnimationLengths.Add(AnimSequence->GetPlayLength());
+
+							AnimInstance->InitializeAnimation(false);
 						}
 					}
 					else if (AnimSequence != nullptr)
@@ -2600,4 +2602,4 @@ void UAnimationSharingManager::LogData() const
 			}
 		}
 	}
-}
\ No newline at end of file
+}
