AnimSharing: Register and Unregister functions now return a bool to let the caller know whether the operation succeded or not

diff --git a/Engine/Plugins/Developer/AnimationSharing/Source/AnimationSharing/Private/AnimationSharingManager.cpp b/Engine/Plugins/Developer/AnimationSharing/Source/AnimationSharing/Private/AnimationSharingManager.cpp
--- a/Engine/Plugins/Developer/AnimationSharing/Source/AnimationSharing/Private/AnimationSharingManager.cpp
+++ b/Engine/Plugins/Developer/AnimationSharing/Source/AnimationSharing/Private/AnimationSharingManager.cpp
@@ -500,11 +500,11 @@ void UAnimationSharingManager::Tick(float DeltaTime)
 
 }
 
-void UAnimationSharingManager::RegisterActor(AActor* InActor, FUpdateActorHandle CallbackDelegate)
+bool UAnimationSharingManager::RegisterActor(AActor* InActor, FUpdateActorHandle CallbackDelegate)
 {
 	if (!UAnimationSharingManager::AnimationSharingEnabled())
 	{
-		return;
+		return false;
 	}
 
 	if (InActor)
@@ -537,15 +537,17 @@ void UAnimationSharingManager::RegisterActor(AActor* InActor, FUpdateActorHandle
 			return CurrentSkeleton;
 		}();
 
-		RegisterActorWithSkeleton(InActor, UsedSkeleton, CallbackDelegate);
+		return RegisterActorWithSkeleton(InActor, UsedSkeleton, CallbackDelegate);
 	}
+
+	return false;
 }
 
-void UAnimationSharingManager::RegisterActorWithSkeleton(AActor* InActor, const USkeleton* SharingSkeleton, FUpdateActorHandle CallbackDelegate)
+bool UAnimationSharingManager::RegisterActorWithSkeleton(AActor* InActor, const USkeleton* SharingSkeleton, FUpdateActorHandle CallbackDelegate)
 {
 	if (!UAnimationSharingManager::AnimationSharingEnabled())
 	{
-		return;
+		return false;
 	}
 
 	const uint32 InstanceIdx = [this, SharingSkeleton]() -> uint32
@@ -603,22 +605,27 @@ void UAnimationSharingManager::RegisterActorWithSkeleton(AActor* InActor, const
 			const int32 ActorHandle = CreateActorHandle(IntCastChecked<uint8>(Data->SkeletonID), ActorIndex);
 
 			ActorData.UpdateActorHandleDelegate.ExecuteIfBound(ActorHandle);
+
+			return true;
 		}
 	}
+
+	return false;
 }
 
-void UAnimationSharingManager::RegisterActorWithSkeletonBP(AActor* InActor, const USkeleton* SharingSkeleton)
+bool UAnimationSharingManager::RegisterActorWithSkeletonBP(AActor* InActor, const USkeleton* SharingSkeleton)
 {
-	RegisterActorWithSkeleton(InActor, SharingSkeleton, FUpdateActorHandle::CreateLambda([](int32 A) {}));
+	return RegisterActorWithSkeleton(InActor, SharingSkeleton, FUpdateActorHandle::CreateLambda([](int32 A) {}));
 }
 
-void UAnimationSharingManager::UnregisterActor(AActor* InActor)
+bool UAnimationSharingManager::UnregisterActor(AActor* InActor)
 {
 	if (!UAnimationSharingManager::AnimationSharingEnabled())
 	{
-		return;
+		return false;
 	}
 
+	int32 UnregisterCounter = 0;
 	for (int32 SkeletonIndex = 0; SkeletonIndex < PerSkeletonData.Num(); ++SkeletonIndex)
 	{
 		UAnimSharingInstance* SkeletonData = PerSkeletonData[SkeletonIndex];
@@ -696,8 +703,13 @@ void UAnimationSharingManager::UnregisterActor(AActor* InActor)
 
 			SkeletonData->PerActorData.RemoveAtSwap(ActorIndex, EAllowShrinking::No);
 			SkeletonData->RegisteredActors.RemoveAtSwap(ActorIndex, EAllowShrinking::No);
+
+			UnregisterCounter++;
 		}
 	}
+
+	// Return whether or not the actor was unregistered from at least one of the Anim Sharing Instances
+	return UnregisterCounter > 0;
 }
 
 void UAnimationSharingManager::UpdateSignificanceForActorHandle(uint32 InHandle, float InValue)
diff --git a/Engine/Plugins/Developer/AnimationSharing/Source/AnimationSharing/Public/AnimationSharingManager.h b/Engine/Plugins/Developer/AnimationSharing/Source/AnimationSharing/Public/AnimationSharingManager.h
--- a/Engine/Plugins/Developer/AnimationSharing/Source/AnimationSharing/Public/AnimationSharingManager.h
+++ b/Engine/Plugins/Developer/AnimationSharing/Source/AnimationSharing/Public/AnimationSharingManager.h
@@ -522,7 +522,7 @@ public:
 
 	/** Register an Actor with this Animation Sharing manager, according to the SharingSkeleton */
 	UFUNCTION(BlueprintCallable, Category = AnimationSharing, meta = (DisplayName = "Register Actor"))
-	UE_API void RegisterActorWithSkeletonBP(AActor* InActor, const USkeleton* SharingSkeleton);
+	UE_API bool RegisterActorWithSkeletonBP(AActor* InActor, const USkeleton* SharingSkeleton);
 
 	/** Returns whether or not the animation sharing is enabled */
 	UFUNCTION(BlueprintPure, Category = AnimationSharing)
@@ -532,13 +532,13 @@ public:
 	static UE_API UAnimationSharingManager* GetManagerForWorld(UWorld* InWorld);
 
 	/** Registers actor with the animation sharing system */
-	UE_API void RegisterActor(AActor* InActor, FUpdateActorHandle Delegate);
+	UE_API bool RegisterActor(AActor* InActor, FUpdateActorHandle Delegate);
 
 	/** Registers actor with the animation sharing system according to the SharingSkeleton's sharing setup (if available) */
-	UE_API void RegisterActorWithSkeleton(AActor* InActor, const USkeleton* SharingSkeleton, FUpdateActorHandle Delegate);
+	UE_API bool RegisterActorWithSkeleton(AActor* InActor, const USkeleton* SharingSkeleton, FUpdateActorHandle Delegate);
 
 	/** Unregisters actor with the animation sharing system */
-	UE_API void UnregisterActor(AActor* InActor);
+	UE_API bool UnregisterActor(AActor* InActor);
 
 	/** Update cached significance for registered actor */
 	UE_API void UpdateSignificanceForActorHandle(uint32 InHandle, float InValue);
