Added option to anim sharing to randomize the selection of the available leader components so any of the animation setups can be picked at random, rather than the arbitrary First available component in the array.
#rnx

diff --git a/Engine/Plugins/Developer/AnimationSharing/Source/AnimationSharing/Private/AnimationSharingManager.cpp b/Engine/Plugins/Developer/AnimationSharing/Source/AnimationSharing/Private/AnimationSharingManager.cpp
--- a/Engine/Plugins/Developer/AnimationSharing/Source/AnimationSharing/Private/AnimationSharingManager.cpp
+++ b/Engine/Plugins/Developer/AnimationSharing/Source/AnimationSharing/Private/AnimationSharingManager.cpp
@@ -1048,6 +1048,7 @@ void UAnimSharingInstance::SetupState(FPerStateData& StateData, const FAnimation
 	StateData.bReturnToPreviousState = StateEntry.bReturnToPreviousState;
 	StateData.bShouldForwardToState = StateEntry.bSetNextState;
 	StateData.ForwardStateValue = StateEntry.NextState;
+	StateData.bShouldRandomizeOnDemandAnims = StateEntry.bRandomizeOnDemandAnims;
 
 	int32 MaximumNumberOfConcurrentInstances = StateEntry.MaximumNumberOfConcurrentInstances.Default;
 #if WITH_EDITOR
@@ -2260,7 +2261,32 @@ uint32 UAnimSharingInstance::SetupOnDemandInstance(uint8 StateIndex)
 		// Otherwise we'll need to kick one of right now so try and set one up		
 		if (StateData.Components.Num())
 		{
-			const uint32 AvailableIndex = StateData.InUseComponentFrameBits.FindAndSetFirstZeroBit();
+			auto FindRandomAvailableIndex = [&StateData]() -> uint32
+				{
+					TBitArray<>& ComponentFrameBits = StateData.InUseComponentFrameBits;
+					TArray<uint32> ZeroBitIndices;
+					ZeroBitIndices.Reserve(ComponentFrameBits.Num());
+
+					for (int32 Index = 0; Index < ComponentFrameBits.Num(); ++Index)
+					{
+						if (ComponentFrameBits[Index] == 0)
+						{
+							ZeroBitIndices.Add(Index);
+						}
+					}
+
+					if (ZeroBitIndices.Num() > 0)
+					{
+						const uint32 RandomIndex = FMath::RandHelper(ZeroBitIndices.Num());
+						ComponentFrameBits[ZeroBitIndices[RandomIndex]] = 1;
+						return ZeroBitIndices[RandomIndex];
+					}
+
+					// No available components this frame.
+					return INDEX_NONE;
+				};
+
+			const uint32 AvailableIndex = StateData.bShouldRandomizeOnDemandAnims && StateData.InUseComponentFrameBits.Num() > 1 ? FindRandomAvailableIndex() : StateData.InUseComponentFrameBits.FindAndSetFirstZeroBit();
 			
 			if (AvailableIndex != INDEX_NONE)
 			{
diff --git a/Engine/Plugins/Developer/AnimationSharing/Source/AnimationSharing/Public/AnimationSharingManager.h b/Engine/Plugins/Developer/AnimationSharing/Source/AnimationSharing/Public/AnimationSharingManager.h
--- a/Engine/Plugins/Developer/AnimationSharing/Source/AnimationSharing/Public/AnimationSharingManager.h
+++ b/Engine/Plugins/Developer/AnimationSharing/Source/AnimationSharing/Public/AnimationSharingManager.h
@@ -176,7 +176,7 @@ struct FInstanceStack
 /** Structure which holds data about a unique state which is linked to an enumeration value defined by the user. The data is populated from the user exposed FAnimationStateEntry */
 struct FPerStateData
 {
-	FPerStateData() : bIsOnDemand(false), bIsAdditive(false), BlendTime(0.f), CurrentFrameOnDemandIndex(INDEX_NONE), StateEnumValue(INDEX_NONE), AdditiveAnimationSequence(nullptr) {}
+	FPerStateData() : bIsOnDemand(false), bIsAdditive(false), bShouldRandomizeOnDemandAnims(false), BlendTime(0.f), CurrentFrameOnDemandIndex(INDEX_NONE), StateEnumValue(INDEX_NONE), AdditiveAnimationSequence(nullptr) {}
 
 	/** Flag whether or not this state is an on-demand state, this means that we kick off a unique animation when needed */
 	bool bIsOnDemand;
@@ -190,6 +190,9 @@ struct FPerStateData
 	/** Flag whether or not ForwardStateValue should be used hwen the animation has finished*/
 	bool bShouldForwardToState;
 
+	/* Flag to allow random selection of OnDemand components so we don't always get only the first avaialable component. */
+	bool bShouldRandomizeOnDemandAnims;
+
 	/** Duration of blending when blending to this state */
 	float BlendTime;
 
diff --git a/Engine/Plugins/Developer/AnimationSharing/Source/AnimationSharing/Public/AnimationSharingTypes.h b/Engine/Plugins/Developer/AnimationSharing/Source/AnimationSharing/Public/AnimationSharingTypes.h
--- a/Engine/Plugins/Developer/AnimationSharing/Source/AnimationSharing/Public/AnimationSharingTypes.h
+++ b/Engine/Plugins/Developer/AnimationSharing/Source/AnimationSharing/Public/AnimationSharingTypes.h
@@ -45,7 +45,7 @@ struct FAnimationStateEntry
 {
 	GENERATED_BODY()
 public:
-	FAnimationStateEntry() : State(0), bOnDemand(false), bAdditive(false), BlendTime(0.f), bReturnToPreviousState(false), bSetNextState(false), NextState(0), MaximumNumberOfConcurrentInstances(1), WiggleTimePercentage(0.1f), bRequiresCurves(false) {}
+	FAnimationStateEntry() : State(0), bOnDemand(false), bRandomizeOnDemandAnims(false), bAdditive(false), BlendTime(0.f), bReturnToPreviousState(false), bSetNextState(false), NextState(0), MaximumNumberOfConcurrentInstances(1), WiggleTimePercentage(0.1f), bRequiresCurves(false) {}
 
 	/** Enum value linked to this state */
 	UPROPERTY(EditAnywhere, Category = AnimationSharing)
@@ -59,6 +59,10 @@ public:
 	UPROPERTY(EditAnywhere, Category = AnimationSharing)
 	bool bOnDemand;
 
+	/* Flag that determines whether or not we attempt to randomize the selection of OnDemand animations, if the total number of components for the OnDemand state is more than 1. */
+	UPROPERTY(EditAnywhere, Category = AnimSharing, meta = (EditCondition = "bOnDemand"))
+	bool bRandomizeOnDemandAnims;
+
 	/** Whether or not this state is an additive state */
 	UPROPERTY(EditAnywhere, Category = AnimationSharing, meta = (EditCondition = "bOnDemand"))
 	bool bAdditive;
