Ensure we do not try and destroy sharing actor in case initialization has previously failed (actor == nullptr)

diff --git a/Engine/Plugins/Developer/AnimationSharing/Source/AnimationSharing/Private/AnimationSharingManager.cpp b/Engine/Plugins/Developer/AnimationSharing/Source/AnimationSharing/Private/AnimationSharingManager.cpp
--- a/Engine/Plugins/Developer/AnimationSharing/Source/AnimationSharing/Private/AnimationSharingManager.cpp
+++ b/Engine/Plugins/Developer/AnimationSharing/Source/AnimationSharing/Private/AnimationSharingManager.cpp
@@ -874,7 +874,10 @@ void UAnimSharingInstance::BeginDestroy()
 
 	if (UWorld* World = GetWorld())
 	{
-		World->DestroyActor(SharingActor);
+		if (SharingActor)
+		{
+			World->DestroyActor(SharingActor);
+		}
 	}
 
 	PerStateData.Empty();
