Add per platform option in animsharing for sending curves during blends.

diff --git a/Engine/Plugins/Developer/AnimationSharing/Source/AnimationSharing/Private/AnimationSharingManager.cpp b/Engine/Plugins/Developer/AnimationSharing/Source/AnimationSharing/Private/AnimationSharingManager.cpp
--- a/Engine/Plugins/Developer/AnimationSharing/Source/AnimationSharing/Private/AnimationSharingManager.cpp
+++ b/Engine/Plugins/Developer/AnimationSharing/Source/AnimationSharing/Private/AnimationSharingManager.cpp
@@ -233,6 +233,7 @@ void UAnimationSharingManager::AddSetup(const UAnimationSharingSetup* InSetup)
 			NewScalabilitySettings.BlendSignificanceValue = NewScalabilitySettings.BlendSignificanceValue.GetValueForPlatform(PlatformName);
 			NewScalabilitySettings.MaximumNumberConcurrentBlends = NewScalabilitySettings.MaximumNumberConcurrentBlends.GetValueForPlatform(PlatformName);
 			NewScalabilitySettings.TickSignificanceValue = NewScalabilitySettings.TickSignificanceValue.GetValueForPlatform(PlatformName);
+			NewScalabilitySettings.RequireCurvesDuringBlends = NewScalabilitySettings.RequireCurvesDuringBlends.GetValueForPlatform(PlatformName);
 	#endif
 		
 			for (const FPerSkeletonAnimationSharingSetup& SkeletonSetup : InSetup->SkeletonSetups)
@@ -1014,6 +1015,8 @@ bool UAnimSharingInstance::Setup(UAnimationSharingManager* AnimationSharingManag
 				BlendComponent->SetVisibility(GLeaderComponentsVisible == 1);
 				BlendComponent->bEnableMaterialParameterCaching = SkeletonSetup.bEnableMaterialParameterCaching;
 
+				BlendComponent->bPropagateCurvesToFollowers = ScalabilitySettings.RequireCurvesDuringBlends.Default;
+
 				BlendComponent->PrimaryComponentTick.AddPrerequisite(AnimSharingManager, AnimSharingManager->GetTickFunction());
 
 				FTransitionBlendInstance* BlendActor = new FTransitionBlendInstance();
diff --git a/Engine/Plugins/Developer/AnimationSharing/Source/AnimationSharing/Public/AnimationSharingTypes.h b/Engine/Plugins/Developer/AnimationSharing/Source/AnimationSharing/Public/AnimationSharingTypes.h
--- a/Engine/Plugins/Developer/AnimationSharing/Source/AnimationSharing/Public/AnimationSharingTypes.h
+++ b/Engine/Plugins/Developer/AnimationSharing/Source/AnimationSharing/Public/AnimationSharingTypes.h
@@ -188,6 +188,10 @@ public:
 	/** Significance value tied to whether or not the leader pose components should be ticking */
 	UPROPERTY(EditAnywhere, Category = AnimationSharing, meta = (ClampMin = "0.0", UIMin = "0.0"))
 	FPerPlatformFloat TickSignificanceValue;
+
+	/** Whether or not curves are required during blends */
+	UPROPERTY(EditAnywhere, Category = AnimationSharing)
+	FPerPlatformBool RequireCurvesDuringBlends;
 };
 
 #undef UE_API
