Fixed issue with additive instances not taking into account actor running an on-demand state or being blended between states.
* It now checks them in order of blending, on-demand, regular state
* In case somehow no valid base component is found, the additive state request will be ignored (warning outputted with info)
Fixed issue where on-demand blend out behaviour did not set blend instance on actor-state (bBlending would be true but blend-instance-index == INDEX_NONE)
Added some more details to state-logging formats
Changes being submitted on behalf of [REDACTED_HANDLE]

diff --git a/Engine/Plugins/Developer/AnimationSharing/Source/AnimationSharing/Private/AnimationSharingManager.cpp b/Engine/Plugins/Developer/AnimationSharing/Source/AnimationSharing/Private/AnimationSharingManager.cpp
--- a/Engine/Plugins/Developer/AnimationSharing/Source/AnimationSharing/Private/AnimationSharingManager.cpp
+++ b/Engine/Plugins/Developer/AnimationSharing/Source/AnimationSharing/Private/AnimationSharingManager.cpp
@@ -1144,7 +1144,7 @@ void UAnimSharingInstance::SetupState(FPerStateData& StateData, const FAnimation
 							AnimInstance->StateIndex = StateEntry.State;
 							AnimInstance->ComponentIndex = IntCastChecked<uint8>(Components.Num());
 
-							/** Set the current animation length length */
+							/** Set the current animation length */
 							StateData.AnimationLengths.Add(AnimSequence->GetPlayLength());
 
 							AnimInstance->InitializeAnimation(false);
@@ -1196,7 +1196,7 @@ void UAnimSharingInstance::SetupState(FPerStateData& StateData, const FAnimation
 					AdditiveInstance->Initialise(AdditiveComponent, SkeletonSetup.AdditiveAnimBlueprint.Get());
 					StateData.AdditiveInstanceStack.AddInstance(AdditiveInstance);
 					
-					/** Set the current animation length length */
+					/** Set the current animation length  */
 					StateData.AnimationLengths.Add(AnimSequence->GetPlayLength());
 					Components.Add(AdditiveComponent);
 				}
@@ -1461,7 +1461,7 @@ void UAnimSharingInstance::TickOnDemandInstances()
 			SetComponentUsage(false, Instance.State, Instance.UsedPerStateComponentIndex);
 
 #if LOG_STATES 
-			UE_LOG(LogAnimationSharing, Log, TEXT("Finished on demand %s"), *StateEnum->GetDisplayNameTextByValue(Instance.State).ToString());
+			UE_LOG(LogAnimationSharing, Log, TEXT("Finished on demand [%i] %s"), InstanceIndex, *StateEnum->GetDisplayNameTextByValue(Instance.State).ToString());
 #endif
 			auto SetActorState = [this, &Instance](uint32 ActorIndex, uint8 NewState)
 			{
@@ -1603,6 +1603,8 @@ void UAnimSharingInstance::TickOnDemandInstances()
 					{
 						// TODO what if two actors have a different state they are blending to? --> Store permutation index
 						Instance.BlendToPermutationIndex = BlendInstances[BlendIndex].ToPermutationIndex;
+						// Make sure to update actor data with running blend instance (index)
+						PerActorData[ActorIndex].BlendInstanceIndex = BlendIndex;
 #if LOG_STATES 
 						UE_LOG(LogAnimationSharing, Log, TEXT("Blending [%i] out from %s to %s"), ActorIndex, *StateEnum->GetDisplayNameTextByValue(Instance.State).ToString(), *StateEnum->GetDisplayNameTextByValue(BlendToState).ToString());
 #endif
@@ -1714,7 +1716,7 @@ void UAnimSharingInstance::TickActorStates()
 				{
 					UpdateState();
 #if LOG_STATES 
-					UE_LOG(LogAnimationSharing, Log, TEXT("Changing state to %s from %s while running on demand %i"), *StateEnum->GetDisplayNameTextByValue(CurrentState).ToString(), *StateEnum->GetDisplayNameTextByValue(ActorData.PreviousState).ToString(), ActorIndex);
+					UE_LOG(LogAnimationSharing, Log, TEXT("Changing state to %s from %s while running on demand %i (%i)"), *StateEnum->GetDisplayNameTextByValue(CurrentState).ToString(), *StateEnum->GetDisplayNameTextByValue(ActorData.PreviousState).ToString(), ActorIndex, ActorData.PermutationIndex);
 #endif
 				}
 				/** Play additive animation only if actor isn't already playing one */
@@ -1722,12 +1724,46 @@ void UAnimSharingInstance::TickActorStates()
 				{
 					if (!ActorData.bRunningAdditive)
 					{
-						const uint32 AdditiveInstanceIndex = SetupAdditiveInstance(CurrentState, PreviousState, ActorData.PermutationIndex);
-						if (AdditiveInstanceIndex != INDEX_NONE)
+						USkeletalMeshComponent* BaseComponent = nullptr;
+						if (ActorData.bBlending)
 						{
-							ActorData.bRunningAdditive = true;
-							ActorData.AdditiveInstanceIndex = AdditiveInstanceIndex;
-							AdditiveInstances[AdditiveInstanceIndex].ActorIndex = ActorIndex;
+							// When running a blend (either between regular states, or in/out of on-demand states), use the blend component as the additive base
+							if (BlendInstances.IsValidIndex(ActorData.BlendInstanceIndex))
+							{
+								BaseComponent = BlendInstances[ActorData.BlendInstanceIndex].TransitionBlendInstance->GetComponent();
+							}
+							
+						}
+						else if (ActorData.bRunningOnDemand)
+						{
+							// When running an on-demand, use the per-state-component used by/running the on demand instance as the additive base
+							if(OnDemandInstances.IsValidIndex(ActorData.OnDemandInstanceIndex))
+							{
+								const FOnDemandInstance& RunningOnDemandInstance = OnDemandInstances[ActorData.OnDemandInstanceIndex];
+								BaseComponent = PerStateData[RunningOnDemandInstance.State].Components[RunningOnDemandInstance.UsedPerStateComponentIndex];
+							}
+						
+						}
+						else if (PerStateData[PreviousState].Components.IsValidIndex(ActorData.PermutationIndex))
+						{
+							// When running a regular state, use the permutation component as the additive base 
+							BaseComponent = PerStateData[PreviousState].Components[ActorData.PermutationIndex];
+						}
+
+						// If we somehow did not manage to find a valid base component, ignore the additive request
+						if (BaseComponent)
+						{
+							const uint32 AdditiveInstanceIndex = SetupAdditiveInstance(CurrentState, BaseComponent);
+							if (AdditiveInstanceIndex != INDEX_NONE)
+							{
+								ActorData.bRunningAdditive = true;
+								ActorData.AdditiveInstanceIndex = AdditiveInstanceIndex;
+								AdditiveInstances[AdditiveInstanceIndex].ActorIndex = ActorIndex;
+							}
+						}
+						else
+						{
+							UE_LOG(LogAnimationSharing, Warning, TEXT("Failed to setup additive state %s for %i. OnDemand %i %i. Blend %i %i. State %i %i"), *StateEnum->GetDisplayNameTextByValue(CurrentState).ToString(), ActorIndex, ActorData.bRunningOnDemand, ActorData.OnDemandInstanceIndex, ActorData.bBlending, ActorData.BlendInstanceIndex, PreviousState, ActorData.PermutationIndex);
 						}
 					}
 				}
@@ -1863,7 +1899,7 @@ void UAnimSharingInstance::TickBlendInstances()
 		if (Instance.EndTime <= WorldTime)
 		{
 #if LOG_STATES
-			UE_LOG(LogAnimationSharing, Log, TEXT("Finished blend %s from %s"), *StateEnum->GetDisplayNameTextByValue(Instance.StateTo).ToString(), *StateEnum->GetDisplayNameTextByValue(Instance.StateFrom).ToString());
+			UE_LOG(LogAnimationSharing, Log, TEXT("Finished blend %s from %s [%i]"), *StateEnum->GetDisplayNameTextByValue(Instance.StateTo).ToString(), *StateEnum->GetDisplayNameTextByValue(Instance.StateFrom).ToString(), InstanceIndex);
 #endif
 
 			// Finish blend into unique animation, need to just set it to use the correct component
@@ -1878,7 +1914,7 @@ void UAnimSharingInstance::TickBlendInstances()
 					SetLeaderComponentForActor(ActorIndex, Instance.TransitionBlendInstance->GetToComponent());
 					PerActorData[ActorIndex].PermutationIndex = 0;
 #if LOG_STATES
-					UE_LOG(LogAnimationSharing, Log, TEXT("Setting %i to on-demand component %i"), ActorIndex, Instance.ToOnDemandInstanceIndex);
+					UE_LOG(LogAnimationSharing, Log, TEXT("Setting %i to on-demand component %i [%i]"), ActorIndex, Instance.ToOnDemandInstanceIndex, InstanceIndex);
 #endif
 
 					for (uint32 ComponentIndex : PerActorData[ActorIndex].ComponentIndices)
@@ -1899,7 +1935,7 @@ void UAnimSharingInstance::TickBlendInstances()
 					if (PerActorData[ActorIndex].CurrentState == Instance.StateTo)
 					{
 #if LOG_STATES 
-						UE_LOG(LogAnimationSharing, Log, TEXT("Setting %i to state %i | %i"), ActorIndex, Instance.StateTo, Instance.ToPermutationIndex);
+						UE_LOG(LogAnimationSharing, Log, TEXT("Setting %i to state %i | %i [%i]"), ActorIndex, Instance.StateTo, Instance.ToPermutationIndex, InstanceIndex);
 #endif
 						SetPermutationFollowerComponent(Instance.StateTo, ActorIndex, Instance.ToPermutationIndex);
 #if !UE_BUILD_SHIPPING
@@ -1923,7 +1959,7 @@ void UAnimSharingInstance::TickBlendInstances()
 			for (uint32 ActorIndex : Instance.ActorIndices)
 			{
 				PerActorData[ActorIndex].BlendInstanceIndex = INDEX_NONE;
-				PerActorData[ActorIndex].bBlending = 0;
+				PerActorData[ActorIndex].bBlending = false;
 			}
 
 			// Remove this blend instance as it has finished
@@ -2069,7 +2105,11 @@ void UAnimSharingInstance::SetLeaderComponentForActor(uint32 ActorIndex, USkelet
 	const FPerActorData& ActorData = PerActorData[ActorIndex];
 	// Do not update the component of the additive actor itself, otherwise update the base component
 	if (ActorData.bRunningAdditive && AdditiveInstances.IsValidIndex(ActorData.AdditiveInstanceIndex) && AdditiveInstances[ActorData.AdditiveInstanceIndex].AdditiveAnimationInstance->GetComponent() != Component)
-	{		
+	{
+#if LOG_STATES
+		UE_LOG(LogAnimationSharing, Log, TEXT("Updating additive instance base component %s for %i"), *StateEnum->GetDisplayNameTextByValue(AdditiveInstances[ActorData.AdditiveInstanceIndex].State).ToString(), ActorIndex);
+#endif // LOG_STATES
+
 		AdditiveInstances[ActorData.AdditiveInstanceIndex].BaseComponent = Component;
 		AdditiveInstances[ActorData.AdditiveInstanceIndex].AdditiveAnimationInstance->UpdateBaseComponent(Component);
 
@@ -2321,7 +2361,7 @@ uint32 UAnimSharingInstance::SetupOnDemandInstance(uint8 StateIndex)
 				FreeComponent->SetPosition(0.f, false);
 				FreeComponent->Play(false);
 #if LOG_STATES 
-				UE_LOG(LogAnimationSharing, Log, TEXT("Setup on demand state %s"), *StateEnum->GetDisplayNameTextByValue(StateIndex).ToString());
+				UE_LOG(LogAnimationSharing, Log, TEXT("Setup on demand state %s %i"), *StateEnum->GetDisplayNameTextByValue(StateIndex).ToString(), InstanceIndex);
 #endif
 			}
 			else
@@ -2373,7 +2413,7 @@ uint32 UAnimSharingInstance::SetupOnDemandInstance(uint8 StateIndex)
 	return InstanceIndex;
 }
 
-uint32 UAnimSharingInstance::SetupAdditiveInstance(uint8 StateIndex, uint8 FromState, uint8 StateComponentIndex)
+uint32 UAnimSharingInstance::SetupAdditiveInstance(uint8 StateIndex, USkeletalMeshComponent* InBaseComponent)
 {
 	uint32 InstanceIndex = INDEX_NONE;
 
@@ -2384,7 +2424,7 @@ uint32 UAnimSharingInstance::SetupAdditiveInstance(uint8 StateIndex, uint8 FromS
 		FAdditiveInstance& Instance = AdditiveInstances.AddDefaulted_GetRef();
 		Instance.bActive = false;
 		Instance.AdditiveAnimationInstance = AnimationInstance;
-		Instance.BaseComponent = PerStateData[FromState].Components[StateComponentIndex];
+		Instance.BaseComponent = InBaseComponent;
 		const float WorldTimeSeconds = static_cast<float>(GetWorld()->GetTimeSeconds());
 		Instance.EndTime = WorldTimeSeconds + StateData.AdditiveAnimationSequence->GetPlayLength();
 		Instance.State = StateIndex;
diff --git a/Engine/Plugins/Developer/AnimationSharing/Source/AnimationSharing/Public/AnimationSharingManager.h b/Engine/Plugins/Developer/AnimationSharing/Source/AnimationSharing/Public/AnimationSharingManager.h
--- a/Engine/Plugins/Developer/AnimationSharing/Source/AnimationSharing/Public/AnimationSharingManager.h
+++ b/Engine/Plugins/Developer/AnimationSharing/Source/AnimationSharing/Public/AnimationSharingManager.h
@@ -319,7 +319,9 @@ public:
 	UE_API uint32 SetupOnDemandInstance(uint8 StateIndex);
 
 	/** Retrieves an additive instance, these are unique and cannot be reused */
-	UE_API uint32 SetupAdditiveInstance(uint8 StateIndex, uint8 FromState, uint8 StateComponentIndex);
+	UE_DEPRECATED(5.8, "SetupAdditiveInstance with state and component index has been deprecated as it causes invalid behaviour, use new signature instead")
+	uint32 SetupAdditiveInstance(uint8 StateIndex, uint8 FromState, uint8 StateComponentIndex) { return INDEX_NONE; }
+	UE_API uint32 SetupAdditiveInstance(uint8 StateIndex, USkeletalMeshComponent* InBaseComponent);
 
 	/** Retrieves the blend-time for this specific state */
 	UE_API float CalculateBlendTime(uint8 StateIndex) const;
