Ensure that Additive and Blend components have their LOD forced to 0, otherwise it will be impacted by view distance calculations and can cause a situation where the evaluated LOD is lower than the shown LOD for the followed (AI) skeletal mesh components. Causing incorrect bone transforms.

diff --git a/Engine/Plugins/Developer/AnimationSharing/Source/AnimationSharing/Private/AnimationSharingManager.cpp b/Engine/Plugins/Developer/AnimationSharing/Source/AnimationSharing/Private/AnimationSharingManager.cpp
--- a/Engine/Plugins/Developer/AnimationSharing/Source/AnimationSharing/Private/AnimationSharingManager.cpp
+++ b/Engine/Plugins/Developer/AnimationSharing/Source/AnimationSharing/Private/AnimationSharingManager.cpp
@@ -1120,6 +1120,7 @@ void UAnimSharingInstance::SetupState(FPerStateData& StateData, const FAnimation
 					/** Set shared skeletal mesh */
 					Component->SetSkeletalMesh(SkeletalMesh);
 					Component->VisibilityBasedAnimTickOption = EVisibilityBasedAnimTickOption::AlwaysTickPoseAndRefreshBones;
+					// Force component to always run at LOD0, to ensure all follower-components will always be at the same or higher LOD level for bone transforms
 					Component->SetForcedLOD(1);
 					Component->SetVisibility(GLeaderComponentsVisible == 1);
 					Component->bPropagateCurvesToFollowers = StateEntry.bRequiresCurves;
@@ -1187,6 +1188,8 @@ void UAnimSharingInstance::SetupState(FPerStateData& StateData, const FAnimation
 					AdditiveComponent->RegisterComponent();
 					AdditiveComponent->SetRelativeLocation(SpawnLocation);
 					AdditiveComponent->SetSkeletalMesh(SkeletalMesh);
+					// Force component to always run at LOD0, to ensure all follower-components will always be at the same or higher LOD level for bone transforms
+					AdditiveComponent->SetForcedLOD(1);
 					AdditiveComponent->SetVisibility(GLeaderComponentsVisible == 1);
 					AdditiveComponent->bEnableMaterialParameterCaching = SkeletonSetup.bEnableMaterialParameterCaching;
 
diff --git a/Engine/Plugins/Developer/AnimationSharing/Source/AnimationSharing/Private/TransitionBlendInstance.cpp b/Engine/Plugins/Developer/AnimationSharing/Source/AnimationSharing/Private/TransitionBlendInstance.cpp
--- a/Engine/Plugins/Developer/AnimationSharing/Source/AnimationSharing/Private/TransitionBlendInstance.cpp
+++ b/Engine/Plugins/Developer/AnimationSharing/Source/AnimationSharing/Private/TransitionBlendInstance.cpp
@@ -17,9 +17,10 @@ void FTransitionBlendInstance::Initialise(USkeletalMeshComponent* InSkeletalMesh
 			TransitionInstance = Cast<UAnimSharingTransitionInstance>(SkeletalMeshComponent->GetAnimInstance());
 		}
 		
-		SkeletalMeshComponent->SetComponentTickEnabled(false);	
-		SkeletalMeshComponent->SetForcedLOD(0);
-		SkeletalMeshComponent->VisibilityBasedAnimTickOption = EVisibilityBasedAnimTickOption::AlwaysTickPoseAndRefreshBones;		
+		SkeletalMeshComponent->SetComponentTickEnabled(false);
+		// Force component to always run at LOD0, to ensure all follower-components will always be at the same or higher LOD level for bone transforms
+		SkeletalMeshComponent->SetForcedLOD(1);
+		SkeletalMeshComponent->VisibilityBasedAnimTickOption = EVisibilityBasedAnimTickOption::AlwaysTickPoseAndRefreshBones;
 	}
 }
 
@@ -45,7 +46,7 @@ void FTransitionBlendInstance::Setup(USkeletalMeshComponent* InFromComponent, US
 		TransitionInstance->bBlendBool = bBlendState = !TransitionInstance->bBlendBool;
 
 		SkeletalMeshComponent->AddTickPrerequisiteComponent(FromComponent);
-		SkeletalMeshComponent->AddTickPrerequisiteComponent(ToComponent);
+		SkeletalMeshComponent->AddTickPrerequisiteComponent(ToComponent);		
 	}	
 }
 
